AnimNext state-tree plugin setup

diff --git a/Engine/Plugins/Experimental/AnimNextStateTree/AnimNextStateTree.uplugin b/Engine/Plugins/Experimental/AnimNextStateTree/AnimNextStateTree.uplugin
new file mode 100644
--- /dev/null
+++ b/Engine/Plugins/Experimental/AnimNextStateTree/AnimNextStateTree.uplugin
@@ -0,0 +1,59 @@
+{
+	"FileVersion" : 3,
+	"Version" : 1,
+	"VersionName" : "0.1",
+	"FriendlyName" : "Anim Next State Tree",
+	"Description" : "StateTree integration for AnimNext.",
+	"Category" : "Animation",
+	"CreatedBy" : "Epic Games, Inc.",
+	"CreatedByURL" : "https://epicgames.com",
+	"DocsURL" : "",
+	"MarketplaceURL" : "",
+	"SupportURL" : "",
+	"EnabledByDefault" : false,
+	"CanContainContent" : false,
+	"IsBetaVersion": true,
+	"IsExperimentalVersion": true,
+	"Installed" : false,
+	"Plugins":
+	[
+		{
+			"Name": "StateTree",
+			"Enabled": true
+		},
+		{
+			"Name": "Workspace",
+			"Enabled": true
+		},
+		{
+			"Name": "AnimNext",
+			"Enabled": true
+		},		
+		{
+			"Name": "AnimNextAnimGraph",
+			"Enabled": true
+		},
+		{
+			"Name": "StructUtils",
+			"Enabled": true
+		}
+	],
+	"Modules" :
+	[
+		{
+			"Name" : "AnimNextStateTree",
+			"Type" : "Runtime",
+			"LoadingPhase" : "PreDefault"
+		},
+		{
+			"Name" : "AnimNextStateTreeUncookedOnly",
+			"Type" : "UncookedOnly",
+			"LoadingPhase" : "Default"
+		},
+		{
+			"Name" : "AnimNextStateTreeEditor",
+			"Type" : "Editor",
+			"LoadingPhase" : "Default"
+		}
+	]
+}
diff --git a/Engine/Plugins/Experimental/AnimNextStateTree/Source/AnimNextStateTree/AnimNextStateTree.Build.cs b/Engine/Plugins/Experimental/AnimNextStateTree/Source/AnimNextStateTree/AnimNextStateTree.Build.cs
new file mode 100644
--- /dev/null
+++ b/Engine/Plugins/Experimental/AnimNextStateTree/Source/AnimNextStateTree/AnimNextStateTree.Build.cs
@@ -0,0 +1,27 @@
+// Copyright Epic Games, Inc. All Rights Reserved.
+
+namespace UnrealBuildTool.Rules
+{
+	public class AnimNextStateTree : ModuleRules
+	{
+		public AnimNextStateTree(ReadOnlyTargetRules Target) : base(Target)
+		{
+			PrivateDependencyModuleNames.AddRange(
+				new string[]
+				{
+					"AnimNext", 
+					"AnimNextAnimGraph",
+					"RigVM"
+				});
+			PublicDependencyModuleNames.AddRange(
+				new string[]
+				{
+					"Core",
+					"CoreUObject",
+					"StateTreeModule",
+					"Engine"
+				}
+			);
+		}
+	}
+}
\ No newline at end of file
diff --git a/Engine/Plugins/Experimental/AnimNextStateTree/Source/AnimNextStateTree/Private/AnimNextStateTreeModule.cpp b/Engine/Plugins/Experimental/AnimNextStateTree/Source/AnimNextStateTree/Private/AnimNextStateTreeModule.cpp
new file mode 100644
--- /dev/null
+++ b/Engine/Plugins/Experimental/AnimNextStateTree/Source/AnimNextStateTree/Private/AnimNextStateTreeModule.cpp
@@ -0,0 +1,18 @@
+// Copyright Epic Games, Inc. All Rights Reserved.
+
+#include "AnimNextStateTreeModule.h"
+
+namespace UE::AnimNext::StateTree
+{
+
+void FAnimNextStateTreeModule::ShutdownModule()
+{
+}
+
+void FAnimNextStateTreeModule::StartupModule()
+{
+}
+
+IMPLEMENT_MODULE(FAnimNextStateTreeModule, AnimNextStateTree)
+
+}
\ No newline at end of file
diff --git a/Engine/Plugins/Experimental/AnimNextStateTree/Source/AnimNextStateTree/Private/AnimNextStateTreeModule.h b/Engine/Plugins/Experimental/AnimNextStateTree/Source/AnimNextStateTree/Private/AnimNextStateTreeModule.h
new file mode 100644
--- /dev/null
+++ b/Engine/Plugins/Experimental/AnimNextStateTree/Source/AnimNextStateTree/Private/AnimNextStateTreeModule.h
@@ -0,0 +1,17 @@
+// Copyright Epic Games, Inc. All Rights Reserved.
+
+#pragma once
+
+#include "IAnimNextStateTreeModule.h"
+
+namespace UE::AnimNext::StateTree
+{
+
+class FAnimNextStateTreeModule : public IAnimNextStateTreeModule
+{
+public:
+	virtual void ShutdownModule() override;
+	virtual void StartupModule() override;
+};
+
+}
\ No newline at end of file
diff --git a/Engine/Plugins/Experimental/AnimNextStateTree/Source/AnimNextStateTree/Public/IAnimNextStateTreeModule.h b/Engine/Plugins/Experimental/AnimNextStateTree/Source/AnimNextStateTree/Public/IAnimNextStateTreeModule.h
new file mode 100644
--- /dev/null
+++ b/Engine/Plugins/Experimental/AnimNextStateTree/Source/AnimNextStateTree/Public/IAnimNextStateTreeModule.h
@@ -0,0 +1,15 @@
+// Copyright Epic Games, Inc. All Rights Reserved.
+
+#pragma once
+
+#include "CoreMinimal.h"
+#include "Modules/ModuleManager.h"
+
+namespace UE::AnimNext::StateTree
+{
+
+class IAnimNextStateTreeModule : public IModuleInterface
+{
+};
+
+}
diff --git a/Engine/Plugins/Experimental/AnimNextStateTree/Source/AnimNextStateTreeEditor/AnimNextStateTreeEditor.Build.cs b/Engine/Plugins/Experimental/AnimNextStateTree/Source/AnimNextStateTreeEditor/AnimNextStateTreeEditor.Build.cs
new file mode 100644
--- /dev/null
+++ b/Engine/Plugins/Experimental/AnimNextStateTree/Source/AnimNextStateTreeEditor/AnimNextStateTreeEditor.Build.cs
@@ -0,0 +1,28 @@
+// Copyright Epic Games, Inc. All Rights Reserved.
+
+namespace UnrealBuildTool.Rules
+{
+	public class AnimNextStateTreeEditor : ModuleRules
+	{
+		public AnimNextStateTreeEditor(ReadOnlyTargetRules Target) : base(Target)
+		{
+			PrivateDependencyModuleNames.AddRange(new string[] { "AnimNextStateTree", "UnrealEd", "MessageLog" });
+			PublicDependencyModuleNames.AddRange(
+				new string[]
+				{
+					"Core",
+					"CoreUObject",
+					"WorkspaceEditor",
+					"StateTreeEditorModule",
+					"StateTreeModule",
+					"AnimNextStateTreeUncookedOnly",
+					"AnimNextEditor",
+					"SlateCore",
+					"EditorSubsystem",
+					"Engine",
+					"Slate"
+				}
+			);
+		}
+	}
+}
\ No newline at end of file
diff --git a/Engine/Plugins/Experimental/AnimNextStateTree/Source/AnimNextStateTreeEditor/Private/AnimNextStateTreeEditorModule.cpp b/Engine/Plugins/Experimental/AnimNextStateTree/Source/AnimNextStateTreeEditor/Private/AnimNextStateTreeEditorModule.cpp
new file mode 100644
--- /dev/null
+++ b/Engine/Plugins/Experimental/AnimNextStateTree/Source/AnimNextStateTreeEditor/Private/AnimNextStateTreeEditorModule.cpp
@@ -0,0 +1,24 @@
+// Copyright Epic Games, Inc. All Rights Reserved.
+
+#include "AnimNextStateTreeEditorModule.h"
+
+#define LOCTEXT_NAMESPACE "AnimNextStateTreeEditorModule"
+
+namespace UE::AnimNext::StateTree
+{
+
+void FAnimNextStateTreeEditorModule::StartupModule()
+{
+	
+}
+
+void FAnimNextStateTreeEditorModule::ShutdownModule()
+{
+	
+}
+
+IMPLEMENT_MODULE(FAnimNextStateTreeEditorModule, AnimNextStateTreeEditor)
+
+}
+
+#undef LOCTEXT_NAMESPACE // "AnimNextStateTreeEditorModule"
\ No newline at end of file
diff --git a/Engine/Plugins/Experimental/AnimNextStateTree/Source/AnimNextStateTreeEditor/Private/AnimNextStateTreeEditorModule.h b/Engine/Plugins/Experimental/AnimNextStateTree/Source/AnimNextStateTreeEditor/Private/AnimNextStateTreeEditorModule.h
new file mode 100644
--- /dev/null
+++ b/Engine/Plugins/Experimental/AnimNextStateTree/Source/AnimNextStateTreeEditor/Private/AnimNextStateTreeEditorModule.h
@@ -0,0 +1,15 @@
+// Copyright Epic Games, Inc. All Rights Reserved.
+
+#pragma once
+
+#include "IAnimNextStateTreeEditorModule.h"
+
+namespace UE::AnimNext::StateTree
+{
+class FAnimNextStateTreeEditorModule : public IAnimNextStateTreeEditorModule
+{
+public:
+	virtual void ShutdownModule() override;
+	virtual void StartupModule() override;
+};
+}
diff --git a/Engine/Plugins/Experimental/AnimNextStateTree/Source/AnimNextStateTreeEditor/Public/IAnimNextStateTreeEditorModule.h b/Engine/Plugins/Experimental/AnimNextStateTree/Source/AnimNextStateTreeEditor/Public/IAnimNextStateTreeEditorModule.h
new file mode 100644
--- /dev/null
+++ b/Engine/Plugins/Experimental/AnimNextStateTree/Source/AnimNextStateTreeEditor/Public/IAnimNextStateTreeEditorModule.h
@@ -0,0 +1,13 @@
+// Copyright Epic Games, Inc. All Rights Reserved.
+
+#pragma once
+
+#include "CoreMinimal.h"
+#include "Modules/ModuleManager.h"
+
+namespace UE::AnimNext::StateTree
+{
+class IAnimNextStateTreeEditorModule : public IModuleInterface
+{
+};
+}
diff --git a/Engine/Plugins/Experimental/AnimNextStateTree/Source/AnimNextStateTreeUncookedOnly/AnimNextStateTreeUncookedOnly.Build.cs b/Engine/Plugins/Experimental/AnimNextStateTree/Source/AnimNextStateTreeUncookedOnly/AnimNextStateTreeUncookedOnly.Build.cs
new file mode 100644
--- /dev/null
+++ b/Engine/Plugins/Experimental/AnimNextStateTree/Source/AnimNextStateTreeUncookedOnly/AnimNextStateTreeUncookedOnly.Build.cs
@@ -0,0 +1,23 @@
+// Copyright Epic Games, Inc. All Rights Reserved.
+
+namespace UnrealBuildTool.Rules
+{
+	public class AnimNextStateTreeUncookedOnly : ModuleRules
+	{
+		public AnimNextStateTreeUncookedOnly(ReadOnlyTargetRules Target) : base(Target)
+		{
+			PrivateDependencyModuleNames.AddRange(new string[] { "StateTreeEditorModule" });
+			PublicDependencyModuleNames.AddRange(
+				new string[]
+				{
+					"Core",
+					"CoreUObject",
+					"StateTreeModule",
+					"Engine",
+					"WorkspaceEditor",
+					"StructUtils"
+				}
+			);
+		}
+	}
+}
\ No newline at end of file
diff --git a/Engine/Plugins/Experimental/AnimNextStateTree/Source/AnimNextStateTreeUncookedOnly/Private/AnimNextStateTreeUncookedModule.cpp b/Engine/Plugins/Experimental/AnimNextStateTree/Source/AnimNextStateTreeUncookedOnly/Private/AnimNextStateTreeUncookedModule.cpp
new file mode 100644
--- /dev/null
+++ b/Engine/Plugins/Experimental/AnimNextStateTree/Source/AnimNextStateTreeUncookedOnly/Private/AnimNextStateTreeUncookedModule.cpp
@@ -0,0 +1,5 @@
+// Copyright Epic Games, Inc. All Rights Reserved.
+
+#include "AnimNextStateTreeUncookedModule.h"
+
+IMPLEMENT_MODULE(UE::AnimNext::StateTree::FAnimNextStateTreeUncookedModule, AnimNextStateTreeUncookedOnly)
\ No newline at end of file
diff --git a/Engine/Plugins/Experimental/AnimNextStateTree/Source/AnimNextStateTreeUncookedOnly/Private/AnimNextStateTreeUncookedModule.h b/Engine/Plugins/Experimental/AnimNextStateTree/Source/AnimNextStateTreeUncookedOnly/Private/AnimNextStateTreeUncookedModule.h
new file mode 100644
--- /dev/null
+++ b/Engine/Plugins/Experimental/AnimNextStateTree/Source/AnimNextStateTreeUncookedOnly/Private/AnimNextStateTreeUncookedModule.h
@@ -0,0 +1,15 @@
+// Copyright Epic Games, Inc. All Rights Reserved.
+
+#pragma once
+
+#include "CoreMinimal.h"
+#include "Modules/ModuleManager.h"
+
+namespace UE::AnimNext::StateTree
+{
+
+class FAnimNextStateTreeUncookedModule : public IModuleInterface
+{
+};
+
+}
\ No newline at end of file
