Remove "StructUtils" plugin dependency

#rnx

diff --git a/Engine/Plugins/Experimental/AnimNextStateTree/AnimNextStateTree.uplugin b/Engine/Plugins/Experimental/AnimNextStateTree/AnimNextStateTree.uplugin
--- a/Engine/Plugins/Experimental/AnimNextStateTree/AnimNextStateTree.uplugin
+++ b/Engine/Plugins/Experimental/AnimNextStateTree/AnimNextStateTree.uplugin
@@ -32,10 +32,6 @@
 		{
 			"Name": "AnimNextAnimGraph",
 			"Enabled": true
-		},
-		{
-			"Name": "StructUtils",
-			"Enabled": true
 		}
 	],
 	"Modules" :
diff --git a/Engine/Plugins/Experimental/AnimNextStateTree/Source/AnimNextStateTreeUncookedOnly/AnimNextStateTreeUncookedOnly.Build.cs b/Engine/Plugins/Experimental/AnimNextStateTree/Source/AnimNextStateTreeUncookedOnly/AnimNextStateTreeUncookedOnly.Build.cs
--- a/Engine/Plugins/Experimental/AnimNextStateTree/Source/AnimNextStateTreeUncookedOnly/AnimNextStateTreeUncookedOnly.Build.cs
+++ b/Engine/Plugins/Experimental/AnimNextStateTree/Source/AnimNextStateTreeUncookedOnly/AnimNextStateTreeUncookedOnly.Build.cs
@@ -14,8 +14,7 @@ namespace UnrealBuildTool.Rules
 					"CoreUObject",
 					"StateTreeModule",
 					"Engine",
-					"WorkspaceEditor",
-					"StructUtils"
+					"WorkspaceEditor"
 				}
 			);
 		}
