- Ensure that StateTree works with data interfaces, and have them forwarded to instantiated graphs
- Fix crash on drag/drop data-interface variables into a graph

diff --git a/Engine/Plugins/Experimental/AnimNext/Source/AnimNextEditor/Private/Variables/Outliner/VariablesOutlinerMode.cpp b/Engine/Plugins/Experimental/AnimNext/Source/AnimNextEditor/Private/Variables/Outliner/VariablesOutlinerMode.cpp
--- a/Engine/Plugins/Experimental/AnimNext/Source/AnimNextEditor/Private/Variables/Outliner/VariablesOutlinerMode.cpp
+++ b/Engine/Plugins/Experimental/AnimNext/Source/AnimNextEditor/Private/Variables/Outliner/VariablesOutlinerMode.cpp
@@ -343,19 +343,19 @@ TSharedPtr<FDragDropOperation> FVariablesOutlinerMode::CreateDragDropOperation(c
 	{
 		return nullptr;
 	}
-
+	
 	UAnimNextVariableEntry* Entry = Cast<UAnimNextVariableEntry>(VariableItem->WeakEntry.Get());
 	if(Entry == nullptr)
 	{
 		return nullptr;
 	}
 
-	URigVMHost* RigVMHost = Entry->GetTypedOuter<URigVMHost>();
+	URigVMHost* RigVMHost = VariableItem->WeakDataInterfaceEntry.IsValid() ? VariableItem->WeakDataInterfaceEntry.Get()->GetTypedOuter<URigVMHost>() : Entry->GetTypedOuter<URigVMHost>();
 	if(RigVMHost == nullptr)
 	{
 		return nullptr;
 	}
-
+	 
 	FName VariableName = Entry->GetVariableName();
 	TArray<FRigVMExternalVariable> ExternalVariables = RigVMHost->GetVM()->GetExternalVariables(RigVMHost->GetRigVMExtendedExecuteContext());
 	for (const FRigVMExternalVariable& ExternalVariable : ExternalVariables)
diff --git a/Engine/Plugins/Experimental/AnimNextStateTree/Source/AnimNextStateTree/Private/AnimNextStateTreeContext.cpp b/Engine/Plugins/Experimental/AnimNextStateTree/Source/AnimNextStateTree/Private/AnimNextStateTreeContext.cpp
--- a/Engine/Plugins/Experimental/AnimNextStateTree/Source/AnimNextStateTree/Private/AnimNextStateTreeContext.cpp
+++ b/Engine/Plugins/Experimental/AnimNextStateTree/Source/AnimNextStateTree/Private/AnimNextStateTreeContext.cpp
@@ -2,12 +2,13 @@
 
 #include "AnimNextStateTreeContext.h"
 
+#include "DataInterface/DataInterfaceStructAdapter.h"
 #include "TraitCore/TraitStackBinding.h"
 #include "TraitInterfaces/IBlendStack.h"
 #include "TraitInterfaces/ISmoothBlend.h"
 #include "TraitInterfaces/ITimeline.h"
 
-bool FAnimNextStateTreeTraitContext::PushAnimationGraphOntoBlendStack(TNonNullPtr<UAnimNextAnimationGraph> InAnimationGraph, const FAlphaBlendArgs& InBlendArguments) const
+bool FAnimNextStateTreeTraitContext::PushAnimationGraphOntoBlendStack(TNonNullPtr<UAnimNextAnimationGraph> InAnimationGraph, const FAlphaBlendArgs& InBlendArguments, TArrayView<const FInstancedStruct> PayloadStructs) const
 {		
 	if (Binding && Context)
 	{
@@ -25,6 +26,8 @@ bool FAnimNextStateTreeTraitContext::PushAnimationGraphOntoBlendStack(TNonNullPt
 			UE::AnimNext::IBlendStack::FGraphRequest NewGraphRequest;
 			NewGraphRequest.BlendArgs = InBlendArguments;
 			NewGraphRequest.AnimationGraph = InAnimationGraph.Get();
+			NewGraphRequest.Payloads.Append(PayloadStructs);
+			
 			BlendStackBinding.PushGraph(*Context, MoveTemp(NewGraphRequest));
 
 			return true;
diff --git a/Engine/Plugins/Experimental/AnimNextStateTree/Source/AnimNextStateTree/Private/AnimNextStateTreeGraphInstanceTask.cpp b/Engine/Plugins/Experimental/AnimNextStateTree/Source/AnimNextStateTree/Private/AnimNextStateTreeGraphInstanceTask.cpp
--- a/Engine/Plugins/Experimental/AnimNextStateTree/Source/AnimNextStateTree/Private/AnimNextStateTreeGraphInstanceTask.cpp
+++ b/Engine/Plugins/Experimental/AnimNextStateTree/Source/AnimNextStateTree/Private/AnimNextStateTreeGraphInstanceTask.cpp
@@ -3,6 +3,7 @@
 #include "Tasks/AnimNextStateTreeGraphInstanceTask.h"
 #include "StateTreeLinker.h"
 #include "StateTreeExecutionContext.h"
+#include "DataInterface/DataInterfaceStructAdapter.h"
 
 FAnimNextStateTreeGraphInstanceTask::FAnimNextStateTreeGraphInstanceTask()
 {
@@ -18,7 +19,11 @@ EStateTreeRunStatus FAnimNextStateTreeGraphInstanceTask::EnterState(FStateTreeEx
 {
 	FAnimNextStateTreeTraitContext& ExecContext = Context.GetExternalData(TraitContextHandle);
 	FInstanceDataType& InstanceData = Context.GetInstanceData(*this);
-	if (ExecContext.PushAnimationGraphOntoBlendStack(InstanceData.AnimationGraph.Get(), InstanceData.BlendOptions))
+
+	FInstancedStruct Payload;
+	Payload.InitializeAs(Context.GetMutableInstanceData()->GetMutableStorage().GetMutableGlobalParameters().GetScriptStruct(), Context.GetMutableInstanceData()->GetMutableStorage().GetMutableGlobalParameters().GetMemory());
+
+	if(ExecContext.PushAnimationGraphOntoBlendStack(InstanceData.AnimationGraph.Get(), InstanceData.BlendOptions, MakeConstArrayView({Payload})))
 	{
 		return EStateTreeRunStatus::Running;	
 	}
diff --git a/Engine/Plugins/Experimental/AnimNextStateTree/Source/AnimNextStateTree/Private/AnimStateTreeTrait.cpp b/Engine/Plugins/Experimental/AnimNextStateTree/Source/AnimNextStateTree/Private/AnimStateTreeTrait.cpp
--- a/Engine/Plugins/Experimental/AnimNextStateTree/Source/AnimNextStateTree/Private/AnimStateTreeTrait.cpp
+++ b/Engine/Plugins/Experimental/AnimNextStateTree/Source/AnimNextStateTree/Private/AnimStateTreeTrait.cpp
@@ -60,10 +60,10 @@ namespace UE::AnimNext
 				const int32 NumVariables = StateTreeParameters.GetNumPropertiesInBag();
 				for(int32 VariableIndex = 0; VariableIndex < NumVariables; ++VariableIndex)
 				{
-					if (ExtendedExecuteContext.ExternalVariableRuntimeData.IsValidIndex(VariableIndex))
+					if (ExtendedExecuteContext.ExternalVariableRuntimeData.IsValidIndex(VariableIndex) && MutableVariables.GetMemory())
 					{
 						const FPropertyBagPropertyDesc& Desc = StateTreeParameters.GetPropertyBagStruct()->GetPropertyDescs()[VariableIndex];
-						void* TargetAddress = MutableVariables.GetMemory() + Desc.CachedProperty->GetOffset_ForInternal();
+						void* TargetAddress = Desc.CachedProperty->ContainerPtrToValuePtr<void*>(MutableVariables.GetMemory());
 						void const* SourceAddress = ExtendedExecuteContext.ExternalVariableRuntimeData[VariableIndex].Memory;
 
 						Desc.CachedProperty->CopyCompleteValue(TargetAddress, SourceAddress);
@@ -105,7 +105,7 @@ namespace UE::AnimNext
 			
 			FStateTreeExecutionContext StateTreeExecutionContext(*Owner, *InstanceData->StateTree, InstanceData->InstanceData);
 			FAnimNextGraphInstance* OwnerGraphInstance = &Binding.GetTraitPtr().GetNodeInstance()->GetOwner();
-			if(OwnerGraphInstance && InstanceData->StateTree->GetDefaultParameters().Identical(&OwnerGraphInstance->GetVariables(), 0))
+			if(OwnerGraphInstance && InstanceData->StateTree->GetDefaultParameters().GetPropertyBagStruct() == OwnerGraphInstance->GetVariables().GetPropertyBagStruct())
 			{
 				FAnimNextStateTreeTraitContext TraitContext(Context, Binding.GetStack());
 				StateTreeExecutionContext.SetContextDataByName(UStateTreeAnimNextSchema::AnimStateTreeExecutionContextName, FStateTreeDataView(FAnimNextStateTreeTraitContext::StaticStruct(), reinterpret_cast<uint8*>(&TraitContext)));
@@ -119,12 +119,11 @@ namespace UE::AnimNext
 				const int32 NumVariables = StateTreeParameters.GetNumPropertiesInBag();
 				for(int32 VariableIndex = 0; VariableIndex < NumVariables; ++VariableIndex)
 				{
-					if (ExtendedExecuteContext.ExternalVariableRuntimeData.IsValidIndex(VariableIndex))
+					if (ExtendedExecuteContext.ExternalVariableRuntimeData.IsValidIndex(VariableIndex) && MutableVariables.GetMemory())
 					{
 						const FPropertyBagPropertyDesc& Desc = StateTreeParameters.GetPropertyBagStruct()->GetPropertyDescs()[VariableIndex];
-						void* TargetAddress = MutableVariables.GetMemory() + Desc.CachedProperty->GetOffset_ForInternal();
+						void* TargetAddress = Desc.CachedProperty->ContainerPtrToValuePtr<void*>(MutableVariables.GetMemory());
 						void const* SourceAddress = ExtendedExecuteContext.ExternalVariableRuntimeData[VariableIndex].Memory;
-
 						Desc.CachedProperty->CopyCompleteValue(TargetAddress, SourceAddress);
 					}
 				}
diff --git a/Engine/Plugins/Experimental/AnimNextStateTree/Source/AnimNextStateTree/Public/AnimNextStateTreeContext.h b/Engine/Plugins/Experimental/AnimNextStateTree/Source/AnimNextStateTree/Public/AnimNextStateTreeContext.h
--- a/Engine/Plugins/Experimental/AnimNextStateTree/Source/AnimNextStateTree/Public/AnimNextStateTreeContext.h
+++ b/Engine/Plugins/Experimental/AnimNextStateTree/Source/AnimNextStateTree/Public/AnimNextStateTreeContext.h
@@ -2,6 +2,9 @@
 
 #pragma once
 
+#include "Graph/AnimNextGraphInstancePtr.h"
+#include "StructUtils/InstancedStruct.h"
+
 #include "AnimNextStateTreeContext.generated.h"
 
 namespace UE::AnimNext
@@ -23,7 +26,7 @@ struct FAnimNextStateTreeTraitContext
 	
 	FAnimNextStateTreeTraitContext() {}
 
-	bool PushAnimationGraphOntoBlendStack(TNonNullPtr<UAnimNextAnimationGraph> InAnimationGraph, const FAlphaBlendArgs& InBlendArguments) const;
+	bool PushAnimationGraphOntoBlendStack(TNonNullPtr<UAnimNextAnimationGraph> InAnimationGraph, const FAlphaBlendArgs& InBlendArguments, const TArrayView<const FInstancedStruct> PayloadStructs) const;
 	float QueryPlaybackRatio(TNonNullPtr<UAnimNextAnimationGraph> InAnimationGraph) const;
 protected:
 	FAnimNextStateTreeTraitContext(UE::AnimNext::FExecutionContext& InContext, const UE::AnimNext::FTraitStackBinding* InBinding) : Context(&InContext), Binding(InBinding) {}
diff --git a/Engine/Plugins/Experimental/AnimNextStateTree/Source/AnimNextStateTreeUncookedOnly/Internal/AnimNextStateTree_EditorData.cpp b/Engine/Plugins/Experimental/AnimNextStateTree/Source/AnimNextStateTreeUncookedOnly/Internal/AnimNextStateTree_EditorData.cpp
--- a/Engine/Plugins/Experimental/AnimNextStateTree/Source/AnimNextStateTreeUncookedOnly/Internal/AnimNextStateTree_EditorData.cpp
+++ b/Engine/Plugins/Experimental/AnimNextStateTree/Source/AnimNextStateTreeUncookedOnly/Internal/AnimNextStateTree_EditorData.cpp
@@ -15,6 +15,7 @@
 #include "TraitCore/TraitRegistry.h"
 #include "TraitCore/Trait.h"
 #include "AnimStateTreeTrait.h"
+#include "Entries/AnimNextDataInterfaceEntry.h"
 #include "Traits/BlendStackTrait.h"
 
 TSubclassOf<UAssetUserData> UAnimNextStateTree_EditorData::GetAssetUserDataClass() const
@@ -41,6 +42,7 @@ TConstArrayView<TSubclassOf<UAnimNextRigVMAssetEntry>> UAnimNextStateTree_Editor
 {
 		UAnimNextVariableEntry::StaticClass(),
 		UAnimNextAnimationGraphEntry::StaticClass(),
+		UAnimNextDataInterfaceEntry::StaticClass(),
 	};
 
 	return Classes;
