Fix uninitialized FAnimNextStateTreeStateOutlinerData (UStruct) members

diff --git a/Engine/Plugins/Experimental/AnimNextStateTree/Source/AnimNextStateTreeUncookedOnly/Public/AnimNextStateTreeWorkspaceExports.h b/Engine/Plugins/Experimental/AnimNextStateTree/Source/AnimNextStateTreeUncookedOnly/Public/AnimNextStateTreeWorkspaceExports.h
--- a/Engine/Plugins/Experimental/AnimNextStateTree/Source/AnimNextStateTreeUncookedOnly/Public/AnimNextStateTreeWorkspaceExports.h
+++ b/Engine/Plugins/Experimental/AnimNextStateTree/Source/AnimNextStateTreeUncookedOnly/Public/AnimNextStateTreeWorkspaceExports.h
@@ -32,13 +32,13 @@ struct FAnimNextStateTreeStateOutlinerData : public FWorkspaceOutlinerItemData
 	FGuid StateId;
 
 	UPROPERTY()
-	bool bIsLeafState;
+	bool bIsLeafState = true;
 
 	UPROPERTY()
-	EStateTreeStateSelectionBehavior SelectionBehavior;
+	EStateTreeStateSelectionBehavior SelectionBehavior = EStateTreeStateSelectionBehavior::None;
 
 	UPROPERTY()
-	EStateTreeStateType Type;
+	EStateTreeStateType Type = EStateTreeStateType::State;
 
 	UPROPERTY()
 	FSlateColor Color;
