Fix for AnimNextStateTree pluging shutdown crash when UObject has been destroyed before the plugin

diff --git a/Engine/Plugins/Experimental/AnimNextStateTree/Source/AnimNextStateTreeEditor/Private/AnimNextStateTreeEditorModule.cpp b/Engine/Plugins/Experimental/AnimNextStateTree/Source/AnimNextStateTreeEditor/Private/AnimNextStateTreeEditorModule.cpp
--- a/Engine/Plugins/Experimental/AnimNextStateTree/Source/AnimNextStateTreeEditor/Private/AnimNextStateTreeEditorModule.cpp
+++ b/Engine/Plugins/Experimental/AnimNextStateTree/Source/AnimNextStateTreeEditor/Private/AnimNextStateTreeEditorModule.cpp
@@ -261,16 +261,19 @@ void FAnimNextStateTreeEditorModule::StartupModule()
 
 void FAnimNextStateTreeEditorModule::ShutdownModule()
 {
-	if(Workspace::IWorkspaceEditorModule* WorkspaceEditorModule = FModuleManager::Get().GetModulePtr<Workspace::IWorkspaceEditorModule>("WorkspaceEditor"))
+	if (UObjectInitialized())
 	{
-		WorkspaceEditorModule->UnregisterObjectDocumentType(UAnimNextStateTree::StaticClass()->GetClassPathName());
-		WorkspaceEditorModule->UnregisterWorkspaceItemDetails(Workspace::FOutlinerItemDetailsId(FAnimNextStateTreeOutlinerData::StaticStruct()->GetFName()));
-		WorkspaceEditorModule->UnregisterWorkspaceItemDetails(Workspace::FOutlinerItemDetailsId(FAnimNextStateTreeStateOutlinerData::StaticStruct()->GetFName()));
-	}
+		if(Workspace::IWorkspaceEditorModule* WorkspaceEditorModule = FModuleManager::Get().GetModulePtr<Workspace::IWorkspaceEditorModule>("WorkspaceEditor"))
+		{
+				WorkspaceEditorModule->UnregisterObjectDocumentType(UAnimNextStateTree::StaticClass()->GetClassPathName());
+				WorkspaceEditorModule->UnregisterWorkspaceItemDetails(Workspace::FOutlinerItemDetailsId(FAnimNextStateTreeOutlinerData::StaticStruct()->GetFName()));
+				WorkspaceEditorModule->UnregisterWorkspaceItemDetails(Workspace::FOutlinerItemDetailsId(FAnimNextStateTreeStateOutlinerData::StaticStruct()->GetFName()));
+		}
 	
-	if( Editor::IAnimNextEditorModule* AnimNextEditorModule = FModuleManager::Get().GetModulePtr<Editor::IAnimNextEditorModule>("AnimNextEditor"))
-	{
-		AnimNextEditorModule->RemoveWorkspaceSupportedAssetClass(UStateTree::StaticClass()->GetClassPathName());
+		if( Editor::IAnimNextEditorModule* AnimNextEditorModule = FModuleManager::Get().GetModulePtr<Editor::IAnimNextEditorModule>("AnimNextEditor"))
+		{
+			AnimNextEditorModule->RemoveWorkspaceSupportedAssetClass(UStateTree::StaticClass()->GetClassPathName());
+		}
 	}
 }
 
