Fix RigVMWarning on ST override type

#tests Editor verify no warning

diff --git a/Engine/Plugins/Experimental/AnimNextStateTree/Source/AnimNextStateTree/Private/AnimNextStateTreeModule.cpp b/Engine/Plugins/Experimental/AnimNextStateTree/Source/AnimNextStateTree/Private/AnimNextStateTreeModule.cpp
--- a/Engine/Plugins/Experimental/AnimNextStateTree/Source/AnimNextStateTree/Private/AnimNextStateTreeModule.cpp
+++ b/Engine/Plugins/Experimental/AnimNextStateTree/Source/AnimNextStateTree/Private/AnimNextStateTreeModule.cpp
@@ -27,6 +27,7 @@ void FAnimNextStateTreeModule::StartupModule()
 	static UScriptStruct* const AllowedStructTypes[] =
 	{
 		FStateTreeReference::StaticStruct(),
+		FStateTreeReferenceOverrides::StaticStruct(),
 	};
 
 	FRigVMRegistry::Get().RegisterStructTypes(AllowedStructTypes);
