Add support for executing RigVM graphs as conditions.

#tests PIE statetree

diff --git a/Engine/Plugins/Experimental/AnimNextStateTree/Source/AnimNextStateTree/AnimNextStateTree.Build.cs b/Engine/Plugins/Experimental/AnimNextStateTree/Source/AnimNextStateTree/AnimNextStateTree.Build.cs
--- a/Engine/Plugins/Experimental/AnimNextStateTree/Source/AnimNextStateTree/AnimNextStateTree.Build.cs
+++ b/Engine/Plugins/Experimental/AnimNextStateTree/Source/AnimNextStateTree/AnimNextStateTree.Build.cs
@@ -12,7 +12,7 @@ namespace UnrealBuildTool.Rules
 					"AnimNextAnimGraph",
 					"RigVM"
 				});
-			
+
 			PublicDependencyModuleNames.AddRange(
 				new string[]
 				{
@@ -20,9 +20,24 @@ namespace UnrealBuildTool.Rules
 					"CoreUObject",
 					"StateTreeModule",
 					"Engine",
-					"AnimNext", 
+					"AnimNext",
 				}
 			);
+
+			if (Target.bBuildEditor == true)
+			{
+				PrivateDependencyModuleNames.AddRange(
+					new string[]
+					{
+						"AnimNextAnimGraph",
+						"AnimNextAnimGraphUncookedOnly",
+						"AnimNextEditor",
+						"AnimNextUncookedOnly",
+						"RigVMDeveloper",
+						"StateTreeEditorModule",
+						"UnrealEd"
+					});
+			}
 		}
 	}
 }
\ No newline at end of file
diff --git a/Engine/Plugins/Experimental/AnimNextStateTree/Source/AnimNextStateTree/Private/AnimNextStateTreeSchema.cpp b/Engine/Plugins/Experimental/AnimNextStateTree/Source/AnimNextStateTree/Private/AnimNextStateTreeSchema.cpp
--- a/Engine/Plugins/Experimental/AnimNextStateTree/Source/AnimNextStateTree/Private/AnimNextStateTreeSchema.cpp
+++ b/Engine/Plugins/Experimental/AnimNextStateTree/Source/AnimNextStateTree/Private/AnimNextStateTreeSchema.cpp
@@ -12,6 +12,7 @@
 #include "Subsystems/WorldSubsystem.h"
 #include "GameFramework/Actor.h"
 #include "AnimNextStateTreeContext.h"
+#include "Conditions/AnimNextStateTreeRigVMConditionBase.h"
 
 const FName UStateTreeAnimNextSchema::AnimStateTreeExecutionContextName("ExecutionContext");
 
@@ -26,6 +27,7 @@ bool UStateTreeAnimNextSchema::IsStructAllowed(const UScriptStruct* InScriptStru
 	|| InScriptStruct->IsChildOf(FStateTreeEvaluatorCommonBase::StaticStruct())
 	|| InScriptStruct->IsChildOf(FStateTreeTaskCommonBase::StaticStruct())
 	|| InScriptStruct->IsChildOf(FStateTreeConsiderationCommonBase::StaticStruct())
+	|| InScriptStruct->IsChildOf(FAnimNextStateTreeRigVMConditionBase::StaticStruct())
 	|| InScriptStruct->IsChildOf(FStateTreePropertyFunctionCommonBase::StaticStruct())
 	|| InScriptStruct->IsChildOf(FAnimNextStateTreeEvaluatorBase::StaticStruct())
 	|| InScriptStruct->IsChildOf(FAnimNextStateTreeTaskBase::StaticStruct());
diff --git a/Engine/Plugins/Experimental/AnimNextStateTree/Source/AnimNextStateTree/Private/Conditions/AnimNextStateTreeRigVMConditionBase.cpp b/Engine/Plugins/Experimental/AnimNextStateTree/Source/AnimNextStateTree/Private/Conditions/AnimNextStateTreeRigVMConditionBase.cpp
new file mode 100644
--- /dev/null
+++ b/Engine/Plugins/Experimental/AnimNextStateTree/Source/AnimNextStateTree/Private/Conditions/AnimNextStateTreeRigVMConditionBase.cpp
@@ -0,0 +1,364 @@
+// Copyright Epic Games, Inc. All Rights Reserved.
+
+#include "Conditions/AnimNextStateTreeRigVMConditionBase.h"
+#include "Graph/AnimNextGraphInstance.h"
+#include "StateTreeLinker.h"
+#include "StateTreeExecutionContext.h"
+#include "StructUtils/PropertyBag.h"
+#include "TraitCore/ExecutionContext.h"
+
+#if WITH_EDITOR
+#include "AnimGraphUncookedOnlyUtils.h"
+#include "AnimNextRigVMAssetEditorData.h"
+#include "DataInterface/AnimNextDataInterfaceFactory.h"
+#include "DataInterface/AnimNextDataInterface.h"
+#include "DataInterface/AnimNextDataInterface_EditorData.h"
+#include "Factories/Factory.h"
+#include "Graph/AnimNextAnimationGraphSchema.h"
+#include "Graph/RigUnit_AnimNextGraphRoot.h"
+#include "Graph/RigUnit_AnimNextShimRoot.h"
+#include "Param/ParamType.h"
+#include "RigVMBlueprintGeneratedClass.h"
+#include "RigVMFunctions/Execution/RigVMFunction_UserDefinedEvent.h"
+#include "RigVMFunctions/RigVMFunctionDefines.h"
+#include "RigVMModel/Nodes/RigVMFunctionReferenceNode.h"
+#include "RigVMModel/Nodes/RigVMLibraryNode.h"
+#include "RigVMModel/Nodes/RigVMUnitNode.h"
+#include "RigVMModel/Nodes/RigVMVariableNode.h"
+#include "RigVMModel/RigVMFunctionLibrary.h"
+#include "RigVMModel/RigVMGraph.h"
+#include "RigVMTypeUtils.h"
+#include "StateTree.h"
+#include "StateTreeEditorData.h"
+#include "UncookedOnlyUtils.h"
+#endif // WITH_EDITOR
+
+#define LOCTEXT_NAMESPACE "AnimNextStateTreeRigVMConditionBase"
+
+#include UE_INLINE_GENERATED_CPP_BY_NAME(AnimNextStateTreeRigVMConditionBase)
+
+////////////////////////////////////////////////////////////////////////////////////////////////////////////////////////
+// FAnimNextStateTreeRigVMConditionInstanceData
+
+bool FAnimNextStateTreeRigVMConditionInstanceData::IsIndexCacheInitialized() const
+{
+	return VariableIndexCache.bInitialized;
+}
+
+bool FAnimNextStateTreeRigVMConditionInstanceData::TryPopulateIndexCache(TObjectPtr<UAnimNextDataInterface> InDataInterface, FName InResultName, FName InFunctionName)
+{
+	if (VariableIndexCache.bInitialized)
+	{
+		return true;
+	}
+
+	VariableIndexCache.ArgumentIndexes.Reserve(ParamData.Num());
+	TArray<FRigVMExternalVariable> ExternalVariables = InDataInterface->GetExternalVariables();
+	for (const FPropertyBagPropertyDesc& Desc : ParamData.GetPropertyBagStruct()->GetPropertyDescs())
+	{
+		for (int32 Index = 0; Index < ExternalVariables.Num(); Index++)
+		{
+			FName VariableName = ExternalVariables[Index].Name;
+			if (VariableName == FAnimNextStateTreeRigVMConditionBase::MakeInternalVariableName(Desc.Name, InFunctionName))
+			{
+				VariableIndexCache.ArgumentIndexes.Add(Index);
+				break;
+			}
+		}
+	}
+
+	// Search all variables for result, it won't be part of the arg param data
+	bool bResultFound = false;
+	for (int32 Index = 0; Index < ExternalVariables.Num(); Index++)
+	{
+		FName VariableName = ExternalVariables[Index].Name;
+		if (VariableName == InResultName)
+		{
+			VariableIndexCache.ResultIndex = Index;
+			bResultFound = true;
+			break;
+		}
+	}
+
+	bool bAllArgsFound = ensure(VariableIndexCache.ArgumentIndexes.Num() == ParamData.Num());
+	if (bAllArgsFound && ensure(bResultFound))
+	{
+		VariableIndexCache.bInitialized = true;
+	}
+
+	return VariableIndexCache.bInitialized;
+}
+
+TArrayView<const uint8> FAnimNextStateTreeRigVMConditionInstanceData::GetVMArgumentIndexes() const
+{
+	return MakeArrayView(VariableIndexCache.ArgumentIndexes);
+}
+
+uint8 FAnimNextStateTreeRigVMConditionInstanceData::GetVMResultIndex() const
+{
+	return VariableIndexCache.ResultIndex;
+}
+
+////////////////////////////////////////////////////////////////////////////////////////////////////////////////////////
+// FAnimNextStateTreeRigVMConditionBase
+
+const FName FAnimNextStateTreeRigVMConditionBase::InternalEventPrefixName("__Internal__Event__");
+const FName FAnimNextStateTreeRigVMConditionBase::InternalVariablePrefixName("__Internal__Variable__");
+
+FName FAnimNextStateTreeRigVMConditionBase::FAnimNextStateTreeRigVMConditionBase::MakeInternalEventName(const FName InName, const FName InFunctionName)
+{
+	// We need to append the variable name to avoid conflicts with other same named variables used by different functions in the VM (Ex: Result)
+	TStringBuilder<512> StringBuilder;
+	StringBuilder << InternalEventPrefixName.ToString();
+	StringBuilder << TEXT("__");
+	StringBuilder << InFunctionName.ToString();
+	StringBuilder << TEXT("__");
+	StringBuilder << InName.ToString();
+
+	return FName(StringBuilder.ToString());
+}
+
+FName FAnimNextStateTreeRigVMConditionBase::FAnimNextStateTreeRigVMConditionBase::MakeInternalVariableName(const FName InName, const FName InFunctionName)
+{
+	// We need to append the variable name to avoid conflicts with other same named variables used by different functions in the VM (Ex: Result)
+	TStringBuilder<512> StringBuilder;
+	StringBuilder << InternalVariablePrefixName.ToString();
+	StringBuilder << TEXT("__");
+	StringBuilder << InFunctionName.ToString();
+	StringBuilder << TEXT("__");
+	StringBuilder << InName.ToString();
+
+	return FName(StringBuilder.ToString());
+}
+
+bool FAnimNextStateTreeRigVMConditionBase::Link(FStateTreeLinker& Linker)
+{
+	Linker.LinkExternalData(TraitContextHandle);
+	return true;
+}
+
+bool FAnimNextStateTreeRigVMConditionBase::TestCondition(FStateTreeExecutionContext& Context) const
+{
+	if (FAnimNextStateTreeRigVMConditionInstanceData* InstanceData = Context.GetInstanceDataPtr<FAnimNextStateTreeRigVMConditionInstanceData>(*this))
+	{
+		FAnimNextStateTreeTraitContext& ExecContext = Context.GetExternalData(TraitContextHandle);
+		const UE::AnimNext::FExecutionContext* AnimExecContext = ExecContext.GetAnimExecuteContext();
+
+		// Need const cast as VM execution isn't const.
+		if (TObjectPtr<UAnimNextDataInterface> RigVMInstance = const_cast<UAnimNextDataInterface*>(AnimExecContext->GetRootGraphInstance().GetDataInterface()))
+		{
+			// Cache RigVM variable indexes, which can vary as they are parent derived
+			if (!InstanceData->IsIndexCacheInitialized() && !InstanceData->TryPopulateIndexCache(RigVMInstance, InternalResultName, RigVMFunctionHeader.Name))
+			{
+				return false;
+			}
+
+			// Prior to execution, copy over ST binding values to RigVM variables
+			TArrayView<const uint8> ArgumentIndexes = InstanceData->GetVMArgumentIndexes();
+			TConstArrayView<FPropertyBagPropertyDesc> ArgumentPropertyDescs = InstanceData->ParamData.GetPropertyBagStruct()->GetPropertyDescs();
+			TArray<FRigVMExternalVariable> ExternalVariables = RigVMInstance->GetExternalVariables();
+			FStructView StateTreeParamData = InstanceData->ParamData.GetMutableValue();
+
+			for (int32 Index = 0; Index < ArgumentIndexes.Num(); Index++)
+			{
+				const FPropertyBagPropertyDesc& ArgumentDesc = ArgumentPropertyDescs[Index];
+				uint8 RigVMExternalVariableIndex = ArgumentIndexes[Index];
+				void* SourceAddress = ArgumentDesc.CachedProperty->ContainerPtrToValuePtr<void*>(StateTreeParamData.GetMemory());
+				void* TargetAddress = ExternalVariables[RigVMExternalVariableIndex].Memory;
+
+				ArgumentDesc.CachedProperty->CopyCompleteValue(TargetAddress, SourceAddress);
+			}
+
+			RigVMInstance->GetVM()->ExecuteVM(RigVMInstance->GetRigVMExtendedExecuteContext(), InternalEventName);
+			bool Result = ExternalVariables[InstanceData->GetVMResultIndex()].GetValue<bool>();
+
+			return Result;
+		}
+	}
+
+	return true;
+}
+
+#if WITH_EDITOR
+FText FAnimNextStateTreeRigVMConditionBase::GetDescription(const FGuid& ID, FStateTreeDataView InstanceDataView, const IStateTreeBindingLookup& BindingLookup, EStateTreeNodeFormatting Formatting) const
+{
+	return LOCTEXT("AnimNextStateTreeConditon_Desc", "RigVM function driven condition");
+}
+
+void FAnimNextStateTreeRigVMConditionBase::PostEditNodeChangeChainProperty(const FPropertyChangedChainEvent& PropertyChangedEvent, FStateTreeDataView InstanceDataView)
+{
+	if (PropertyChangedEvent.GetMemberPropertyName() == GET_MEMBER_NAME_CHECKED(FAnimNextStateTreeRigVMConditionBase, ConditionFunctionName))
+	{
+		// Function Name selected has changed. Update param struct & result / event names used during execution.
+		if (FAnimNextStateTreeRigVMConditionInstanceData* InstanceData = InstanceDataView.GetMutablePtr<FAnimNextStateTreeRigVMConditionInstanceData>())
+		{
+			// @TODO: This relies on the function name being unique (Ex: In a workspace). For now that's okay. Later on however we will want to use a more robust function picker
+			auto GetRigVMFunctionHeader = [](FName InName)
+			{
+				FRigVMGraphFunctionHeader Result = FRigVMGraphFunctionHeader();
+
+				TMap<FAssetData, FRigVMGraphFunctionHeaderArray> FunctionExports;
+				UE::AnimNext::UncookedOnly::FUtils::GetExportedFunctionsFromAssetRegistry(UE::AnimNext::AnimNextPublicGraphFunctionsExportsRegistryTag, FunctionExports);
+				UE::AnimNext::UncookedOnly::FUtils::GetExportedFunctionsFromAssetRegistry(UE::AnimNext::ControlRigAssetPublicGraphFunctionsExportsRegistryTag, FunctionExports);
+
+				for (const auto& Export : FunctionExports.Array())
+				{
+					for (const FRigVMGraphFunctionHeader& FunctionHeader : Export.Value.Headers)
+					{
+						if (InName == FunctionHeader.Name)
+						{
+							Result = FunctionHeader;
+							return Result;
+						}
+					}
+				}
+
+				return Result;
+			};
+
+			RigVMFunctionHeader = GetRigVMFunctionHeader(ConditionFunctionName);
+			InstanceData->ParamData = FRigVMMemoryStorageStruct();
+
+			TArray<FRigVMPropertyDescription> RigVMPropertyDescriptions;
+			for (const FRigVMGraphFunctionArgument& Argument : RigVMFunctionHeader.Arguments)
+			{
+				if (Argument.Direction == ERigVMPinDirection::Input)
+				{
+					FString CppTypeString = Argument.CPPType.ToString();
+					UObject* CppTypeObject = Argument.CPPTypeObject.Get();
+					FRigVMPropertyDescription RigVMPropertyDescription = FRigVMPropertyDescription(Argument.Name, CppTypeString, CppTypeObject, Argument.DefaultValue);
+					RigVMPropertyDescriptions.Add(RigVMPropertyDescription);
+				}
+				else if (Argument.Direction == ERigVMPinDirection::Output)
+				{
+					ensure(Argument.CPPType == RigVMTypeUtils::BoolTypeName);
+					ResultName = Argument.Name;
+					InternalResultName = FAnimNextStateTreeRigVMConditionBase::MakeInternalVariableName(ResultName, RigVMFunctionHeader.Name);
+				}
+			}
+		
+			InstanceData->ParamData.AddProperties(RigVMPropertyDescriptions);
+			InternalEventName = FAnimNextStateTreeRigVMConditionBase::MakeInternalEventName(RigVMFunctionHeader.Name, RigVMFunctionHeader.Name);
+		}
+	}
+}
+
+void FAnimNextStateTreeRigVMConditionBase::GetProgrammaticGraphs(UAnimNextDataInterface_EditorData* EditorData, FRigVMClient& RigVMClient, TArray<URigVMGraph*>& OutGraphs) const
+{
+	URigVMGraph* Graph = NewObject<URigVMGraph>(EditorData, RigVMFunctionHeader.Name, RF_Transient);
+	Graph->SetSchemaClass(UAnimNextAnimationGraphSchema::StaticClass());
+
+	URigVMController* Controller = RigVMClient.GetOrCreateController(Graph);
+	Controller->RemoveNodes(Controller->GetGraph()->GetNodes());
+	
+	// @TODO: Would like to make these variables / procedurally generated ones hidden in AnimNext editor
+	// @TODO: Also need to auto delete these as conditions are removed
+	// Create input variables & result variable
+	{
+		for (const FRigVMGraphFunctionArgument& Argument : RigVMFunctionHeader.Arguments)
+		{
+			FName InternalVariableName = FAnimNextStateTreeRigVMConditionBase::MakeInternalVariableName(Argument.Name, RigVMFunctionHeader.Name);
+			if (Argument.Direction == ERigVMPinDirection::Input)
+			{
+				FRigVMTemplateArgumentType TemplateArgument(Argument.CPPType, Argument.CPPTypeObject.Get());
+				if (!EditorData->FindEntry(InternalVariableName))
+				{
+					EditorData->AddVariable(InternalVariableName, FAnimNextParamType::FromRigVMTemplateArgument(TemplateArgument), Argument.DefaultValue);
+				}
+			}
+			else if (Argument.Direction == ERigVMPinDirection::Output)
+			{
+				FRigVMTemplateArgumentType TemplateArgument(Argument.CPPType, Argument.CPPTypeObject.Get());
+				if (!EditorData->FindEntry(InternalVariableName))
+				{
+					EditorData->AddVariable(InternalVariableName, FAnimNextParamType::FromRigVMTemplateArgument(TemplateArgument), Argument.DefaultValue);
+				}
+			}
+		}
+	}
+
+	// Create RigVM graph for function
+	{
+		FString FunctionNodeString = RigVMFunctionHeader.Name.ToString();
+		bool bSetupUndoRedo = false;
+
+		// Add event entry node, so VM can call our event
+		URigVMUnitNode* EventNode = Controller->AddUnitNode(FRigVMFunction_UserDefinedEvent::StaticStruct(), TEXT("Execute"), FVector2D::ZeroVector, FunctionNodeString, bSetupUndoRedo);
+		if (!EventNode)
+		{
+			return;
+		}
+		URigVMPin* EventNamePin = EventNode->FindPin(GET_MEMBER_NAME_STRING_CHECKED(FRigVMFunction_UserDefinedEvent, EventName));
+		if (EventNamePin == nullptr)
+		{
+			return;
+		}
+		FString EventNameString = InternalEventName.ToString();
+		Controller->SetPinDefaultValue(EventNamePin->GetPinPath(), EventNameString);
+
+		// Add function node
+		URigVMLibraryNode* LibraryNode = Cast<URigVMLibraryNode>(RigVMFunctionHeader.LibraryPointer.GetNodeSoftPath().ResolveObject());
+		if (!LibraryNode)
+		{
+			return;
+		}
+		URigVMFunctionReferenceNode* FunctionNode = Controller->AddFunctionReferenceNode(LibraryNode, FVector2D::ZeroVector, FunctionNodeString, bSetupUndoRedo);
+		if (!FunctionNode)
+		{
+			return;
+		}
+
+		// Generate & link input arguments, also generate result variable node but link later
+		URigVMVariableNode* ResultNode = nullptr;
+		for (const FRigVMGraphFunctionArgument& Argument : RigVMFunctionHeader.Arguments)
+		{
+			bool bIsGetter = Argument.Name != ResultName;
+			FName InternalVariableName = FAnimNextStateTreeRigVMConditionBase::MakeInternalVariableName(Argument.Name, RigVMFunctionHeader.Name);
+			URigVMVariableNode* FunctionArgumentVariableNode = Controller->AddVariableNode(InternalVariableName, Argument.CPPType.ToString(), Argument.CPPTypeObject.Get(), bIsGetter, Argument.DefaultValue, FVector2D::ZeroVector, InternalVariableName.ToString(), bSetupUndoRedo);
+			if (!FunctionArgumentVariableNode)
+			{
+				return;
+			}
+
+			if (bIsGetter)
+			{
+				URigVMPin* ArgumentValuePin = FunctionArgumentVariableNode->GetValuePin();
+				URigVMPin* FunctionArgumentSTring = FunctionNode->FindPin(Argument.Name.ToString());
+
+				if (!Controller->AddLink(ArgumentValuePin, FunctionArgumentSTring, bSetupUndoRedo))
+				{
+					return;
+				}
+			}
+			else
+			{
+				ResultNode = FunctionArgumentVariableNode;
+			}
+		}
+		if (!ResultNode)
+		{
+			return;
+		}
+
+		// Link Execute pins
+		URigVMPin* ExecuteOutputPin = EventNode->FindPin(FRigVMStruct::ExecuteContextName.ToString());
+		URigVMPin* ExecuteInputPin = ResultNode->FindPin(FRigVMStruct::ExecuteContextName.ToString());
+		if (!Controller->AddLink(ExecuteOutputPin, ExecuteInputPin, bSetupUndoRedo))
+		{
+			return;
+		}
+
+		// Link Result pins
+		FString ResultNodeString = ResultName.ToString();
+		URigVMPin* FunctionResultPin = FunctionNode->FindPin(ResultNodeString);
+		URigVMPin* ResultValuePin = ResultNode->GetValuePin();
+		if (!Controller->AddLink(FunctionResultPin, ResultValuePin, bSetupUndoRedo))
+		{
+			return;
+		}
+
+		OutGraphs.Add(Graph);
+	}
+}
+
+#endif // WITH_EDITOR
diff --git a/Engine/Plugins/Experimental/AnimNextStateTree/Source/AnimNextStateTree/Public/AnimNextStateTreeContext.h b/Engine/Plugins/Experimental/AnimNextStateTree/Source/AnimNextStateTree/Public/AnimNextStateTreeContext.h
--- a/Engine/Plugins/Experimental/AnimNextStateTree/Source/AnimNextStateTree/Public/AnimNextStateTreeContext.h
+++ b/Engine/Plugins/Experimental/AnimNextStateTree/Source/AnimNextStateTree/Public/AnimNextStateTreeContext.h
@@ -29,6 +29,8 @@ struct FAnimNextStateTreeTraitContext
 
 	bool PushAssetOntoBlendStack(TNonNullPtr<UObject> InAsset, const FAlphaBlendArgs& InBlendArguments, FAnimNextDataInterfacePayload&& InPayload) const;
 	float QueryPlaybackRatio(TNonNullPtr<UObject> InAsset) const;
+
+	UE::AnimNext::FExecutionContext* GetAnimExecuteContext() { return Context; }
 protected:
 	FAnimNextStateTreeTraitContext(UE::AnimNext::FExecutionContext& InContext, const UE::AnimNext::FTraitBinding* InBinding) : Context(&InContext), Binding(InBinding) {}
 protected:
diff --git a/Engine/Plugins/Experimental/AnimNextStateTree/Source/AnimNextStateTree/Public/Conditions/AnimNextStateTreeRigVMConditionBase.h b/Engine/Plugins/Experimental/AnimNextStateTree/Source/AnimNextStateTree/Public/Conditions/AnimNextStateTreeRigVMConditionBase.h
new file mode 100644
--- /dev/null
+++ b/Engine/Plugins/Experimental/AnimNextStateTree/Source/AnimNextStateTree/Public/Conditions/AnimNextStateTreeRigVMConditionBase.h
@@ -0,0 +1,118 @@
+// Copyright Epic Games, Inc. All Rights Reserved.
+
+#pragma once
+
+#include "AnimNextStateTreeContext.h"
+#include "DataInterface/AnimNextDataInterface.h"
+#include "RigVMCore/RigVMExternalVariable.h"
+#include "RigVMCore/RigVMGraphFunctionDefinition.h"
+#include "RigVMCore/RigVMMemoryStorageStruct.h"
+#include "StateTreeConditionBase.h"
+#include "StateTreeReference.h"
+#include "Templates/SubclassOf.h"
+
+#include "AnimNextStateTreeRigVMConditionBase.generated.h"
+
+class UAnimNextDataInterface;
+class UAnimNextDataInterface_EditorData;
+class URigVMGraph;
+
+struct FStateTreeExecutionContext;
+struct FRigVMClient;
+
+USTRUCT()
+struct ANIMNEXTSTATETREE_API FAnimNextStateTreeRigVMConditionInstanceData
+{
+	GENERATED_BODY()
+
+public:
+
+	/** Per instance copy of ST function arguments we will copy into our RigVM before execution */
+	UPROPERTY(EditAnywhere, Category = "Parameter", Meta = (FixedLayout, ShowOnlyInnerProperties))
+	FRigVMMemoryStorageStruct ParamData = FRigVMMemoryStorageStruct();
+
+public:
+
+	/** True if VM variables are indexed */
+	bool IsIndexCacheInitialized() const;
+
+	/** True if successfully populated index cache from Data Interface, or already populated*/
+	bool TryPopulateIndexCache(TObjectPtr<UAnimNextDataInterface> InDataInterface, FName InResultName, FName InFunctionName);
+
+	/** Gets Argument Indexes, you are responsible for checking if initialized first */
+	TArrayView<const uint8> GetVMArgumentIndexes() const;
+
+	/** Gets Result Index, you are responsible for checking if initialized first */
+	uint8 GetVMResultIndex() const;
+
+protected:
+
+	struct FRigVariableIndexCache
+	{
+		TArray<uint8, TInlineAllocator<8>> ArgumentIndexes = {};
+		uint8 ResultIndex = 0;
+		bool bInitialized = false;
+	};
+
+	/** Helper struct to store parent RigVM variable indicies, */
+	FRigVariableIndexCache VariableIndexCache = FRigVariableIndexCache();
+};
+
+/**
+ * Wrapper for RigVM based Conditions. 
+ */
+USTRUCT()
+struct ANIMNEXTSTATETREE_API FAnimNextStateTreeRigVMConditionBase : public FStateTreeConditionBase
+{
+	GENERATED_BODY()
+
+	using FInstanceDataType = FAnimNextStateTreeRigVMConditionInstanceData;
+
+	static const FName InternalEventPrefixName;
+	static const FName InternalVariablePrefixName;
+	static FName MakeInternalEventName(const FName InName, const FName InFunctionName);
+	static FName MakeInternalVariableName(const FName InName, const FName InFunctionName);
+
+public:
+
+	//~ Begin FStateTreeConditionBase Interface
+	virtual bool Link(FStateTreeLinker& Linker) override;
+	virtual bool TestCondition(FStateTreeExecutionContext& Context) const override;
+	//~ End FStateTreeConditionBase Interface
+
+	//~ Begin FStateTreeNodeBase Interface
+	virtual const UStruct* GetInstanceDataType() const override { return FInstanceDataType::StaticStruct(); }
+#if WITH_EDITOR
+	virtual FText GetDescription(const FGuid& ID, FStateTreeDataView InstanceDataView, const IStateTreeBindingLookup& BindingLookup, EStateTreeNodeFormatting Formatting = EStateTreeNodeFormatting::Text) const override;
+	virtual void PostEditNodeChangeChainProperty(const FPropertyChangedChainEvent& PropertyChangedEvent, FStateTreeDataView InstanceDataView) override;
+#endif
+	//~ End FStateTreeNodeBase Interface
+	
+#if WITH_EDITOR
+	/** Adds condition graphs to parent StateTree RigVM */
+	virtual void GetProgrammaticGraphs(UAnimNextDataInterface_EditorData* EditorData, FRigVMClient& RigVMClient, TArray<URigVMGraph*>& OutGraphs) const;
+#endif
+
+public:
+
+#if WITH_EDITORONLY_DATA
+	UPROPERTY(EditAnywhere, Category = "Parameter", Meta = (GetOptions = "AnimNextStateTreeUncookedOnly.AnimNextStateTreeFunctionLibraryHelper.GetExposedAnimNextFunctionNames"))
+	FName ConditionFunctionName = NAME_None;
+
+	/** Name of the Result of function */
+	UPROPERTY()
+	FName ResultName = NAME_None;
+#endif
+
+	UPROPERTY()
+	FRigVMGraphFunctionHeader RigVMFunctionHeader = FRigVMGraphFunctionHeader();
+
+	UPROPERTY()
+	FName InternalEventName = NAME_None;
+
+	UPROPERTY()
+	FName InternalResultName = NAME_None;
+
+public:
+	TStateTreeExternalDataHandle<FAnimNextStateTreeTraitContext> TraitContextHandle;
+};
\ No newline at end of file
diff --git a/Engine/Plugins/Experimental/AnimNextStateTree/Source/AnimNextStateTreeUncookedOnly/Internal/AnimNextStateTree_EditorData.cpp b/Engine/Plugins/Experimental/AnimNextStateTree/Source/AnimNextStateTreeUncookedOnly/Internal/AnimNextStateTree_EditorData.cpp
--- a/Engine/Plugins/Experimental/AnimNextStateTree/Source/AnimNextStateTreeUncookedOnly/Internal/AnimNextStateTree_EditorData.cpp
+++ b/Engine/Plugins/Experimental/AnimNextStateTree/Source/AnimNextStateTreeUncookedOnly/Internal/AnimNextStateTree_EditorData.cpp
@@ -5,6 +5,7 @@
 #include "AnimGraphUncookedOnlyUtils.h"
 #include "AnimNextStateTree.h"
 #include "AnimNextStateTreeWorkspaceAssetUserData.h"
+#include "Conditions/AnimNextStateTreeRigVMConditionBase.h"
 #include "StateTree.h"
 #include "StateTreeEditorData.h"
 #include "Entries/AnimNextVariableEntry.h"
@@ -124,7 +125,18 @@ void UAnimNextStateTree_EditorData::GetProgrammaticGraphs(const FRigVMCompileSet
 			InSettings.ReportError(TEXT("Failed to link TraitStack and Graph Output pins"));
 			return;
 		}
-		
-		OutGraphs.Add(Graph);	
+
+		// Give child nodes a chance to add graphs
+		for (FConstStructView Node : AnimStateTree->StateTree->GetNodes())
+		{
+			// Manually check for each child node type, UStruct's do not reflect interfaces so we can't use an interface here.
+			if (const FAnimNextStateTreeRigVMConditionBase* Condition = Node.GetPtr<const FAnimNextStateTreeRigVMConditionBase>())
+			{
+				Condition->GetProgrammaticGraphs(this, RigVMClient, OutGraphs);
+			}
+		}
+
+		// @TODO: This has to be the last graph or ST will not execute. Figure out why, graph order should not be implicitly required
+		OutGraphs.Add(Graph);
 	}
 }
diff --git a/Engine/Plugins/Experimental/AnimNextStateTree/Source/AnimNextStateTreeUncookedOnly/Private/AnimNextStateTreeFunctionLibraryHelper.cpp b/Engine/Plugins/Experimental/AnimNextStateTree/Source/AnimNextStateTreeUncookedOnly/Private/AnimNextStateTreeFunctionLibraryHelper.cpp
new file mode 100644
--- /dev/null
+++ b/Engine/Plugins/Experimental/AnimNextStateTree/Source/AnimNextStateTreeUncookedOnly/Private/AnimNextStateTreeFunctionLibraryHelper.cpp
@@ -0,0 +1,25 @@
+// Copyright Epic Games, Inc. All Rights Reserved.
+
+#include "AnimNextStateTreeFunctionLibraryHelper.h"
+#include "RigVMBlueprintGeneratedClass.h"
+#include "UncookedOnlyUtils.h"
+
+const TArray<FName> UAnimNextStateTreeFunctionLibraryHelper::GetExposedAnimNextFunctionNames()
+{
+	// Can't cache here since the list of functions is dynamic.
+	TArray<FName> Result;
+
+	TMap<FAssetData, FRigVMGraphFunctionHeaderArray> FunctionExports;
+	UE::AnimNext::UncookedOnly::FUtils::GetExportedFunctionsFromAssetRegistry(UE::AnimNext::AnimNextPublicGraphFunctionsExportsRegistryTag, FunctionExports);
+	UE::AnimNext::UncookedOnly::FUtils::GetExportedFunctionsFromAssetRegistry(UE::AnimNext::ControlRigAssetPublicGraphFunctionsExportsRegistryTag, FunctionExports);
+	
+	for (const TPair<FAssetData, FRigVMGraphFunctionHeaderArray>& Export : FunctionExports.Array())
+	{
+		for (const FRigVMGraphFunctionHeader& FunctionHeader : Export.Value.Headers)
+		{
+			Result.Add(FunctionHeader.Name);
+		}
+	}
+
+	return Result;
+}
\ No newline at end of file
diff --git a/Engine/Plugins/Experimental/AnimNextStateTree/Source/AnimNextStateTreeUncookedOnly/Public/AnimNextStateTreeFunctionLibraryHelper.h b/Engine/Plugins/Experimental/AnimNextStateTree/Source/AnimNextStateTreeUncookedOnly/Public/AnimNextStateTreeFunctionLibraryHelper.h
new file mode 100644
--- /dev/null
+++ b/Engine/Plugins/Experimental/AnimNextStateTree/Source/AnimNextStateTreeUncookedOnly/Public/AnimNextStateTreeFunctionLibraryHelper.h
@@ -0,0 +1,24 @@
+// Copyright Epic Games, Inc. All Rights Reserved.
+
+#pragma once
+
+#include "AnimNextStateTreeFunctionLibraryHelper.generated.h"
+
+struct FRigVMGraphFunctionHeader;
+
+/**
+ * Helper class for using RigVM functions in ST
+ * 
+ * Currently used for GetOptions funciton population
+ */
+UCLASS()
+class ANIMNEXTSTATETREEUNCOOKEDONLY_API UAnimNextStateTreeFunctionLibraryHelper : public UObject
+{
+	GENERATED_BODY()
+
+public:
+
+	// Get list of function names in module
+	UFUNCTION()
+	static const TArray<FName> GetExposedAnimNextFunctionNames();
+};
\ No newline at end of file
diff --git a/Engine/Plugins/Runtime/RigVM/Source/RigVMEditor/Public/Editor/RigVMVariantDetailCustomization.h b/Engine/Plugins/Runtime/RigVM/Source/RigVMEditor/Public/Editor/RigVMVariantDetailCustomization.h
--- a/Engine/Plugins/Runtime/RigVM/Source/RigVMEditor/Public/Editor/RigVMVariantDetailCustomization.h
+++ b/Engine/Plugins/Runtime/RigVM/Source/RigVMEditor/Public/Editor/RigVMVariantDetailCustomization.h
@@ -34,7 +34,7 @@ protected:
 	void OnAddTag(const FName& InTagName);
 	void OnRemoveTag(const FName& InTagName);
 
-	URigVMBlueprint* BlueprintBeingCustomized;
+	URigVMBlueprint* BlueprintBeingCustomized = nullptr;
 
 	/** The log widget used for function variants */
 	TSharedPtr<SRigVMLogWidget> VariantLog;
