Add missing namespace undef

#tests NonUnity

diff --git a/Engine/Plugins/Experimental/AnimNextStateTree/Source/AnimNextStateTree/Private/Conditions/AnimNextStateTreeRigVMConditionBase.cpp b/Engine/Plugins/Experimental/AnimNextStateTree/Source/AnimNextStateTree/Private/Conditions/AnimNextStateTreeRigVMConditionBase.cpp
--- a/Engine/Plugins/Experimental/AnimNextStateTree/Source/AnimNextStateTree/Private/Conditions/AnimNextStateTreeRigVMConditionBase.cpp
+++ b/Engine/Plugins/Experimental/AnimNextStateTree/Source/AnimNextStateTree/Private/Conditions/AnimNextStateTreeRigVMConditionBase.cpp
@@ -362,3 +362,5 @@ void FAnimNextStateTreeRigVMConditionBase::GetProgrammaticGraphs(UAnimNextDataIn
 }
 
 #endif // WITH_EDITOR
+
+#undef NAMESPACE_LOCTEXT
\ No newline at end of file
