Fix AnimNextStateTree marking its package as dirty when (on load) generating ProgrammaticGraphs

diff --git a/Engine/Plugins/Experimental/AnimNextAnimGraph/Source/AnimNextAnimGraphUncookedOnly/Internal/AnimGraphUncookedOnlyUtils.h b/Engine/Plugins/Experimental/AnimNextAnimGraph/Source/AnimNextAnimGraphUncookedOnly/Internal/AnimGraphUncookedOnlyUtils.h
--- a/Engine/Plugins/Experimental/AnimNextAnimGraph/Source/AnimNextAnimGraphUncookedOnly/Internal/AnimGraphUncookedOnlyUtils.h
+++ b/Engine/Plugins/Experimental/AnimNextAnimGraph/Source/AnimNextAnimGraphUncookedOnly/Internal/AnimGraphUncookedOnlyUtils.h
@@ -13,7 +13,7 @@ namespace UE::AnimNext::UncookedOnly
 struct ANIMNEXTANIMGRAPHUNCOOKEDONLY_API FAnimGraphUtils
 {
 	/** Set up a simple animation graph */
-	static void SetupAnimGraph(const FName EntryName, URigVMController* InController);
+	static void SetupAnimGraph(const FName EntryName, URigVMController* InController, bool bSetupUndoRedo = true);
 
 	/** Check whether the supplied node is a trait stack */
 	static bool IsTraitStackNode(URigVMNode* InModelNode);
diff --git a/Engine/Plugins/Experimental/AnimNextAnimGraph/Source/AnimNextAnimGraphUncookedOnly/Private/AnimGraphUncookedOnlyUtils.cpp b/Engine/Plugins/Experimental/AnimNextAnimGraph/Source/AnimNextAnimGraphUncookedOnly/Private/AnimGraphUncookedOnlyUtils.cpp
--- a/Engine/Plugins/Experimental/AnimNextAnimGraph/Source/AnimNextAnimGraphUncookedOnly/Private/AnimGraphUncookedOnlyUtils.cpp
+++ b/Engine/Plugins/Experimental/AnimNextAnimGraph/Source/AnimNextAnimGraphUncookedOnly/Private/AnimGraphUncookedOnlyUtils.cpp
@@ -21,13 +21,13 @@ bool FAnimGraphUtils::IsTraitStackNode(URigVMNode* InModelNode)
 	return false;
 }
 
-void FAnimGraphUtils::SetupAnimGraph(const FName EntryName, URigVMController* InController)
+void FAnimGraphUtils::SetupAnimGraph(const FName EntryName, URigVMController* InController, bool bSetupUndoRedo /*= true*/)
 {
 	// Clear the graph
 	InController->RemoveNodes(InController->GetGraph()->GetNodes());
 
 	// Add root node
-	URigVMUnitNode* MainEntryPointNode = InController->AddUnitNode(FRigUnit_AnimNextGraphRoot::StaticStruct(), FRigUnit_AnimNextGraphRoot::EventName, FVector2D(-400.0f, 0.0f), FString(), false);
+	URigVMUnitNode* MainEntryPointNode = InController->AddUnitNode(FRigUnit_AnimNextGraphRoot::StaticStruct(), FRigUnit_AnimNextGraphRoot::EventName, FVector2D(-400.0f, 0.0f), FString(), bSetupUndoRedo);
 	if(MainEntryPointNode == nullptr)
 	{
 		return;
@@ -45,7 +45,7 @@ void FAnimGraphUtils::SetupAnimGraph(const FName EntryName, URigVMController* In
 		return;
 	}
 
-	InController->SetPinDefaultValue(EntryPointPin->GetPinPath(), EntryName.ToString());
+	InController->SetPinDefaultValue(EntryPointPin->GetPinPath(), EntryName.ToString(), true, bSetupUndoRedo);
 }
 
 }
diff --git a/Engine/Plugins/Experimental/AnimNextStateTree/Source/AnimNextStateTreeUncookedOnly/Internal/AnimNextStateTree_EditorData.cpp b/Engine/Plugins/Experimental/AnimNextStateTree/Source/AnimNextStateTreeUncookedOnly/Internal/AnimNextStateTree_EditorData.cpp
--- a/Engine/Plugins/Experimental/AnimNextStateTree/Source/AnimNextStateTreeUncookedOnly/Internal/AnimNextStateTree_EditorData.cpp
+++ b/Engine/Plugins/Experimental/AnimNextStateTree/Source/AnimNextStateTreeUncookedOnly/Internal/AnimNextStateTree_EditorData.cpp
@@ -58,7 +58,7 @@ void UAnimNextStateTree_EditorData::GetProgrammaticGraphs(const FRigVMCompileSet
 		Graph->SetSchemaClass(UAnimNextAnimationGraphSchema::StaticClass());
 
 		UAnimNextController* Controller = CastChecked<UAnimNextController>(RigVMClient.GetOrCreateController(Graph));
-		UE::AnimNext::UncookedOnly::FAnimGraphUtils::SetupAnimGraph(FRigUnit_AnimNextGraphRoot::DefaultEntryPoint, Controller);
+		UE::AnimNext::UncookedOnly::FAnimGraphUtils::SetupAnimGraph(FRigUnit_AnimNextGraphRoot::DefaultEntryPoint, Controller, false);
 
 		if (Controller->GetGraph()->GetNodes().Num() != 1)
 		{
@@ -83,7 +83,7 @@ void UAnimNextStateTree_EditorData::GetProgrammaticGraphs(const FRigVMCompileSet
 			return;
 		}
 
-		const FName BlendStackTraitName = Controller->AddTraitByName(TraitStackNode->GetFName(), *UE::AnimNext::FTraitRegistry::Get().Find(UE::AnimNext::FBlendStackCoreTrait::TraitUID)->GetTraitName(), INDEX_NONE);
+		const FName BlendStackTraitName = Controller->AddTraitByName(TraitStackNode->GetFName(), *UE::AnimNext::FTraitRegistry::Get().Find(UE::AnimNext::FBlendStackCoreTrait::TraitUID)->GetTraitName(), INDEX_NONE, TEXT(""), false);
 		
 		if (BlendStackTraitName == NAME_None)
 		{
@@ -91,7 +91,7 @@ void UAnimNextStateTree_EditorData::GetProgrammaticGraphs(const FRigVMCompileSet
 			return;
 		}
 		
-		const FName StateTreeTraitName = Controller->AddTraitByName(TraitStackNode->GetFName(), *UE::AnimNext::FTraitRegistry::Get().Find(UE::AnimNext::FStateTreeTrait::TraitUID)->GetTraitName(), INDEX_NONE);
+		const FName StateTreeTraitName = Controller->AddTraitByName(TraitStackNode->GetFName(), *UE::AnimNext::FTraitRegistry::Get().Find(UE::AnimNext::FStateTreeTrait::TraitUID)->GetTraitName(), INDEX_NONE, TEXT(""), false);
 
 		if (StateTreeTraitName == NAME_None)
 		{
@@ -111,7 +111,7 @@ void UAnimNextStateTree_EditorData::GetProgrammaticGraphs(const FRigVMCompileSet
 
 		FString PinValue;
 		FStateTreeReference::StaticStruct()->ExportText(PinValue, &Ref, nullptr, nullptr, 0, nullptr);
-		Controller->SetPinDefaultValue(StateTreeReferencePin->GetPinPath(), PinValue);
+		Controller->SetPinDefaultValue(StateTreeReferencePin->GetPinPath(), PinValue, true, false);
 		
 		URigVMPin* TraitResult = TraitStackNode->FindPin(GET_MEMBER_NAME_STRING_CHECKED(FRigUnit_AnimNextTraitStack, Result));
 		if (TraitResult == nullptr)
