Fix blendstack trait and state tree subgraphs not inheriting data interfaces

diff --git a/Engine/Plugins/Experimental/AnimNextAnimGraph/Source/AnimNextAnimGraph/Private/Graph/AnimNextGraphInstance.cpp b/Engine/Plugins/Experimental/AnimNextAnimGraph/Source/AnimNextAnimGraph/Private/Graph/AnimNextGraphInstance.cpp
--- a/Engine/Plugins/Experimental/AnimNextAnimGraph/Source/AnimNextAnimGraph/Private/Graph/AnimNextGraphInstance.cpp
+++ b/Engine/Plugins/Experimental/AnimNextAnimGraph/Source/AnimNextAnimGraph/Private/Graph/AnimNextGraphInstance.cpp
@@ -298,6 +298,12 @@ void FAnimNextGraphInstance::BindPublicVariables(TArrayView<FStructView> InHostS
 		return;
 	}
 
+	// First apply any automatic bindings to this instance's host
+   	if(FAnimNextDataInterfaceInstance* InstanceHost = GetHost())
+   	{
+   		BindToHostHelper(*InstanceHost, true);
+   	}
+
 	const int32 NumStructs = InHostStructs.Num();
 	if(NumStructs > 0)
 	{
diff --git a/Engine/Plugins/Experimental/AnimNextStateTree/Source/AnimNextStateTree/Private/AnimNextStateTreeGraphInstanceTask.cpp b/Engine/Plugins/Experimental/AnimNextStateTree/Source/AnimNextStateTree/Private/AnimNextStateTreeGraphInstanceTask.cpp
--- a/Engine/Plugins/Experimental/AnimNextStateTree/Source/AnimNextStateTree/Private/AnimNextStateTreeGraphInstanceTask.cpp
+++ b/Engine/Plugins/Experimental/AnimNextStateTree/Source/AnimNextStateTree/Private/AnimNextStateTreeGraphInstanceTask.cpp
@@ -28,9 +28,10 @@ EStateTreeRunStatus FAnimNextStateTreeGraphInstanceTask::EnterState(FStateTreeEx
 	FAnimNextDataInterfacePayload Payload;
 	Payload.AddStructView(InstanceData.Payload.GetMutableValue());
 
-	FInstancedStruct NewNativePayload;
-	NewNativePayload.InitializeAs(Context.GetMutableInstanceData()->GetMutableStorage().GetMutableGlobalParameters().GetScriptStruct(), Context.GetMutableInstanceData()->GetMutableStorage().GetMutableGlobalParameters().GetMemory());
-	Payload.AddNative(MoveTemp(NewNativePayload));
+	// Temporarily disabling parameter overriding because it is blocking data interface variables from being accessible in subgraphs
+	// FInstancedStruct NewNativePayload;
+	// NewNativePayload.InitializeAs(Context.GetMutableInstanceData()->GetMutableStorage().GetMutableGlobalParameters().GetScriptStruct(), Context.GetMutableInstanceData()->GetMutableStorage().GetMutableGlobalParameters().GetMemory());
+	// Payload.AddNative(MoveTemp(NewNativePayload));
 
 	if (ExecContext.PushAssetOntoBlendStack(InstanceData.Asset.Get(), InstanceData.BlendOptions, MoveTemp(Payload)))
 	{
