AnimNext StateTree graph gen improvements
- Add a GUID to conditions to keep variable / binding names unique across entire StateTree
- Add State name to generated variables for easier identification
- Log to compiler when state tree compilation fails

#tests PIE

diff --git a/Engine/Plugins/Experimental/AnimNextStateTree/Source/AnimNextStateTree/Internal/AnimNextStateTreeEditorOnlyTypes.h b/Engine/Plugins/Experimental/AnimNextStateTree/Source/AnimNextStateTree/Internal/AnimNextStateTreeEditorOnlyTypes.h
new file mode 100644
--- /dev/null
+++ b/Engine/Plugins/Experimental/AnimNextStateTree/Source/AnimNextStateTree/Internal/AnimNextStateTreeEditorOnlyTypes.h
@@ -0,0 +1,30 @@
+// Copyright Epic Games, Inc. All Rights Reserved.
+
+#pragma once
+
+#include "CoreMinimal.h"
+
+class UAnimNextDataInterface_EditorData;
+class URigVMGraph;
+
+struct FRigVMCompileSettings;
+struct FRigVMClient;
+
+/**
+ * Helper struct to contain info needed to compile RigVM Graphs
+ */
+struct ANIMNEXTSTATETREE_API FAnimNextStateTreeProgrammaticGraphParams
+{
+	FAnimNextStateTreeProgrammaticGraphParams(UAnimNextDataInterface_EditorData* InEditorData, const FRigVMCompileSettings& InSettings, FRigVMClient& InRigVMClient, TArray<URigVMGraph*>& InOutGraphs)
+		: EditorData(InEditorData)
+		, Settings(InSettings)
+		, RigVMClient(InRigVMClient)
+		, OutGraphs(InOutGraphs)
+	{
+	}
+
+	UAnimNextDataInterface_EditorData* EditorData;
+	const FRigVMCompileSettings& Settings;
+	FRigVMClient& RigVMClient;
+	TArray<URigVMGraph*>& OutGraphs;
+};
diff --git a/Engine/Plugins/Experimental/AnimNextStateTree/Source/AnimNextStateTree/Private/Conditions/AnimNextStateTreeRigVMConditionBase.cpp b/Engine/Plugins/Experimental/AnimNextStateTree/Source/AnimNextStateTree/Private/Conditions/AnimNextStateTreeRigVMConditionBase.cpp
--- a/Engine/Plugins/Experimental/AnimNextStateTree/Source/AnimNextStateTree/Private/Conditions/AnimNextStateTreeRigVMConditionBase.cpp
+++ b/Engine/Plugins/Experimental/AnimNextStateTree/Source/AnimNextStateTree/Private/Conditions/AnimNextStateTreeRigVMConditionBase.cpp
@@ -10,6 +10,7 @@
 #if WITH_EDITOR
 #include "AnimGraphUncookedOnlyUtils.h"
 #include "AnimNextRigVMAssetEditorData.h"
+#include "AnimNextStateTreeEditorOnlyTypes.h"
 #include "DataInterface/AnimNextDataInterfaceFactory.h"
 #include "DataInterface/AnimNextDataInterface.h"
 #include "DataInterface/AnimNextDataInterface_EditorData.h"
@@ -19,8 +20,10 @@
 #include "Graph/RigUnit_AnimNextShimRoot.h"
 #include "Param/ParamType.h"
 #include "RigVMBlueprintGeneratedClass.h"
+#include "RigVMCompiler/RigVMCompiler.h"
 #include "RigVMFunctions/Execution/RigVMFunction_UserDefinedEvent.h"
 #include "RigVMFunctions/RigVMFunctionDefines.h"
+#include "RigVMModel/RigVMClient.h"
 #include "RigVMModel/Nodes/RigVMFunctionReferenceNode.h"
 #include "RigVMModel/Nodes/RigVMLibraryNode.h"
 #include "RigVMModel/Nodes/RigVMUnitNode.h"
@@ -45,7 +48,7 @@ bool FAnimNextStateTreeRigVMConditionInstanceData::IsIndexCacheInitialized() con
 	return VariableIndexCache.bInitialized;
 }
 
-bool FAnimNextStateTreeRigVMConditionInstanceData::TryPopulateIndexCache(TObjectPtr<UAnimNextDataInterface> InDataInterface, FName InResultName, FName InFunctionName)
+bool FAnimNextStateTreeRigVMConditionInstanceData::TryPopulateIndexCache(const FAnimNextStateTreeRigVMConditionBase& OwningCondition, TObjectPtr<UAnimNextDataInterface> InDataInterface, FName InResultName, FName InFunctionName)
 {
 	if (VariableIndexCache.bInitialized)
 	{
@@ -59,7 +62,7 @@ bool FAnimNextStateTreeRigVMConditionInstanceData::TryPopulateIndexCache(TObject
 		for (int32 Index = 0; Index < ExternalVariables.Num(); Index++)
 		{
 			FName VariableName = ExternalVariables[Index].Name;
-			if (VariableName == FAnimNextStateTreeRigVMConditionBase::MakeInternalVariableName(Desc.Name, InFunctionName))
+			if (VariableName == OwningCondition.MakeInternalVariableName(Desc.Name))
 			{
 				VariableIndexCache.ArgumentIndexes.Add(Index);
 				break;
@@ -105,28 +108,42 @@ uint8 FAnimNextStateTreeRigVMConditionInstanceData::GetVMResultIndex() const
 const FName FAnimNextStateTreeRigVMConditionBase::InternalEventPrefixName("__Internal__Event__");
 const FName FAnimNextStateTreeRigVMConditionBase::InternalVariablePrefixName("__Internal__Variable__");
 
-FName FAnimNextStateTreeRigVMConditionBase::FAnimNextStateTreeRigVMConditionBase::MakeInternalEventName(const FName InName, const FName InFunctionName)
+FName FAnimNextStateTreeRigVMConditionBase::MakeInternalEventName(const FName InName) const
 {
+	ensure(!StateName.IsNone());
+	ensure(!RigVMFunctionHeader.Name.IsNone());
+	ensure(NodeId.IsValid());
+
 	// We need to append the variable name to avoid conflicts with other same named variables used by different functions in the VM (Ex: Result)
 	TStringBuilder<512> StringBuilder;
 	StringBuilder << InternalEventPrefixName.ToString();
+	StringBuilder << StateName.ToString();
 	StringBuilder << TEXT("__");
-	StringBuilder << InFunctionName.ToString();
+	StringBuilder << RigVMFunctionHeader.Name.ToString();
 	StringBuilder << TEXT("__");
 	StringBuilder << InName.ToString();
+	StringBuilder << TEXT("__");
+	StringBuilder << NodeId.ToString();
 
 	return FName(StringBuilder.ToString());
 }
 
-FName FAnimNextStateTreeRigVMConditionBase::FAnimNextStateTreeRigVMConditionBase::MakeInternalVariableName(const FName InName, const FName InFunctionName)
+FName FAnimNextStateTreeRigVMConditionBase::MakeInternalVariableName(const FName InName) const
 {
+	ensure(!StateName.IsNone());
+	ensure(!RigVMFunctionHeader.Name.IsNone());
+	ensure(NodeId.IsValid());
+
 	// We need to append the variable name to avoid conflicts with other same named variables used by different functions in the VM (Ex: Result)
 	TStringBuilder<512> StringBuilder;
 	StringBuilder << InternalVariablePrefixName.ToString();
+	StringBuilder << StateName.ToString();
 	StringBuilder << TEXT("__");
-	StringBuilder << InFunctionName.ToString();
+	StringBuilder << RigVMFunctionHeader.Name.ToString();
 	StringBuilder << TEXT("__");
 	StringBuilder << InName.ToString();
+	StringBuilder << TEXT("__");
+	StringBuilder << NodeId.ToString();
 
 	return FName(StringBuilder.ToString());
 }
@@ -148,7 +165,7 @@ bool FAnimNextStateTreeRigVMConditionBase::TestCondition(FStateTreeExecutionCont
 		if (TObjectPtr<UAnimNextDataInterface> RigVMInstance = const_cast<UAnimNextDataInterface*>(AnimExecContext->GetRootGraphInstance().GetDataInterface()))
 		{
 			// Cache RigVM variable indexes, which can vary as they are parent derived
-			if (!InstanceData->IsIndexCacheInitialized() && !InstanceData->TryPopulateIndexCache(RigVMInstance, InternalResultName, RigVMFunctionHeader.Name))
+			if (!InstanceData->IsIndexCacheInitialized() && !InstanceData->TryPopulateIndexCache(*this, RigVMInstance, InternalResultName, RigVMFunctionHeader.Name))
 			{
 				return false;
 			}
@@ -233,18 +250,26 @@ void FAnimNextStateTreeRigVMConditionBase::PostEditNodeChangeChainProperty(const
 				{
 					ensure(Argument.CPPType == RigVMTypeUtils::BoolTypeName);
 					ResultName = Argument.Name;
-					InternalResultName = FAnimNextStateTreeRigVMConditionBase::MakeInternalVariableName(ResultName, RigVMFunctionHeader.Name);
 				}
 			}
 		
 			InstanceData->ParamData.AddProperties(RigVMPropertyDescriptions);
-			InternalEventName = FAnimNextStateTreeRigVMConditionBase::MakeInternalEventName(RigVMFunctionHeader.Name, RigVMFunctionHeader.Name);
 		}
 	}
 }
 
-void FAnimNextStateTreeRigVMConditionBase::GetProgrammaticGraphs(UAnimNextDataInterface_EditorData* EditorData, FRigVMClient& RigVMClient, TArray<URigVMGraph*>& OutGraphs) const
+void FAnimNextStateTreeRigVMConditionBase::GetProgrammaticGraphs(FAnimNextStateTreeProgrammaticGraphParams& InProgrammaticGraphParams, const UStateTreeState* State, const FStateTreeBindableStructDesc& Desc)
 {
+	UAnimNextDataInterface_EditorData* EditorData = InProgrammaticGraphParams.EditorData;
+	const FRigVMCompileSettings& Settings = InProgrammaticGraphParams.Settings;
+	FRigVMClient& RigVMClient = InProgrammaticGraphParams.RigVMClient;
+	TArray<URigVMGraph*>& OutGraphs = InProgrammaticGraphParams.OutGraphs;
+	
+	StateName = State->Name;
+	NodeId = Desc.ID;
+	InternalResultName = MakeInternalVariableName(ResultName);
+	InternalEventName = MakeInternalEventName(RigVMFunctionHeader.Name);
+
 	URigVMGraph* Graph = NewObject<URigVMGraph>(EditorData, RigVMFunctionHeader.Name, RF_Transient);
 	Graph->SetSchemaClass(UAnimNextAnimationGraphSchema::StaticClass());
 
@@ -257,7 +282,7 @@ void FAnimNextStateTreeRigVMConditionBase::GetProgrammaticGraphs(UAnimNextDataIn
 	{
 		for (const FRigVMGraphFunctionArgument& Argument : RigVMFunctionHeader.Arguments)
 		{
-			FName InternalVariableName = FAnimNextStateTreeRigVMConditionBase::MakeInternalVariableName(Argument.Name, RigVMFunctionHeader.Name);
+			FName InternalVariableName = MakeInternalVariableName(Argument.Name);
 			if (Argument.Direction == ERigVMPinDirection::Input)
 			{
 				FRigVMTemplateArgumentType TemplateArgument(Argument.CPPType, Argument.CPPTypeObject.Get());
@@ -286,11 +311,13 @@ void FAnimNextStateTreeRigVMConditionBase::GetProgrammaticGraphs(UAnimNextDataIn
 		URigVMUnitNode* EventNode = Controller->AddUnitNode(FRigVMFunction_UserDefinedEvent::StaticStruct(), TEXT("Execute"), FVector2D::ZeroVector, FunctionNodeString, bSetupUndoRedo);
 		if (!EventNode)
 		{
+			Settings.ReportError(FString::Printf(TEXT("Failed to add Event Node to VM Condition: %s"), *RigVMFunctionHeader.Name.ToString()));
 			return;
 		}
 		URigVMPin* EventNamePin = EventNode->FindPin(GET_MEMBER_NAME_STRING_CHECKED(FRigVMFunction_UserDefinedEvent, EventName));
 		if (EventNamePin == nullptr)
 		{
+			Settings.ReportError(FString::Printf(TEXT("Failed to add Event Name Pin to VM Condition: %s"), *RigVMFunctionHeader.Name.ToString()));
 			return;
 		}
 		FString EventNameString = InternalEventName.ToString();
@@ -300,33 +327,47 @@ void FAnimNextStateTreeRigVMConditionBase::GetProgrammaticGraphs(UAnimNextDataIn
 		URigVMLibraryNode* LibraryNode = Cast<URigVMLibraryNode>(RigVMFunctionHeader.LibraryPointer.GetNodeSoftPath().ResolveObject());
 		if (!LibraryNode)
 		{
+			Settings.ReportError(FString::Printf(TEXT("Failed to get Library Node for VM Condition: %s"), *RigVMFunctionHeader.Name.ToString()));
 			return;
 		}
 		URigVMFunctionReferenceNode* FunctionNode = Controller->AddFunctionReferenceNode(LibraryNode, FVector2D::ZeroVector, FunctionNodeString, bSetupUndoRedo);
 		if (!FunctionNode)
 		{
+			Settings.ReportError(FString::Printf(TEXT("Failed to add Library Node to VM Condition: %s"), *RigVMFunctionHeader.Name.ToString()));
 			return;
 		}
 
+		auto GetNameSafe = [](const UObject* InObject)
+		{
+			if (InObject)
+			{
+				return InObject->GetFName().ToString();
+			}
+
+			return FString("Null Obj");
+		};
+
 		// Generate & link input arguments, also generate result variable node but link later
 		URigVMVariableNode* ResultNode = nullptr;
 		for (const FRigVMGraphFunctionArgument& Argument : RigVMFunctionHeader.Arguments)
 		{
 			bool bIsGetter = Argument.Name != ResultName;
-			FName InternalVariableName = FAnimNextStateTreeRigVMConditionBase::MakeInternalVariableName(Argument.Name, RigVMFunctionHeader.Name);
+			FName InternalVariableName = MakeInternalVariableName(Argument.Name);
 			URigVMVariableNode* FunctionArgumentVariableNode = Controller->AddVariableNode(InternalVariableName, Argument.CPPType.ToString(), Argument.CPPTypeObject.Get(), bIsGetter, Argument.DefaultValue, FVector2D::ZeroVector, InternalVariableName.ToString(), bSetupUndoRedo);
 			if (!FunctionArgumentVariableNode)
 			{
+				Settings.ReportError(FString::Printf(TEXT("Failed to add internal variable node for VM Condition: %s, var: %s"), *RigVMFunctionHeader.Name.ToString(), *InternalVariableName.ToString()));
 				return;
 			}
 
 			if (bIsGetter)
 			{
 				URigVMPin* ArgumentValuePin = FunctionArgumentVariableNode->GetValuePin();
-				URigVMPin* FunctionArgumentSTring = FunctionNode->FindPin(Argument.Name.ToString());
+				URigVMPin* FunctionArgumentPin = FunctionNode->FindPin(Argument.Name.ToString());
 
-				if (!Controller->AddLink(ArgumentValuePin, FunctionArgumentSTring, bSetupUndoRedo))
+				if (!Controller->AddLink(ArgumentValuePin, FunctionArgumentPin, bSetupUndoRedo))
 				{
+					Settings.ReportError(FString::Printf(TEXT("Failed to link pins for VM Condition: %s -> %s"), *GetNameSafe(ArgumentValuePin), *GetNameSafe(FunctionArgumentPin)));
 					return;
 				}
 			}
@@ -337,6 +378,7 @@ void FAnimNextStateTreeRigVMConditionBase::GetProgrammaticGraphs(UAnimNextDataIn
 		}
 		if (!ResultNode)
 		{
+			Settings.ReportError(FString::Printf(TEXT("Failed to add internal result node for VM Condition: %s, result: %s"), *RigVMFunctionHeader.Name.ToString(), *InternalResultName.ToString()));
 			return;
 		}
 
@@ -345,6 +387,7 @@ void FAnimNextStateTreeRigVMConditionBase::GetProgrammaticGraphs(UAnimNextDataIn
 		URigVMPin* ExecuteInputPin = ResultNode->FindPin(FRigVMStruct::ExecuteContextName.ToString());
 		if (!Controller->AddLink(ExecuteOutputPin, ExecuteInputPin, bSetupUndoRedo))
 		{
+			Settings.ReportError(FString::Printf(TEXT("Failed to link execute pins for VM Condition: %s -> %s"), *GetNameSafe(ExecuteOutputPin), *GetNameSafe(ExecuteInputPin)));
 			return;
 		}
 
@@ -354,6 +397,7 @@ void FAnimNextStateTreeRigVMConditionBase::GetProgrammaticGraphs(UAnimNextDataIn
 		URigVMPin* ResultValuePin = ResultNode->GetValuePin();
 		if (!Controller->AddLink(FunctionResultPin, ResultValuePin, bSetupUndoRedo))
 		{
+			Settings.ReportError(FString::Printf(TEXT("Failed to link result pins for VM Condition: %s -> %s"), *GetNameSafe(FunctionResultPin), *GetNameSafe(ResultValuePin)));
 			return;
 		}
 
diff --git a/Engine/Plugins/Experimental/AnimNextStateTree/Source/AnimNextStateTree/Public/Conditions/AnimNextStateTreeRigVMConditionBase.h b/Engine/Plugins/Experimental/AnimNextStateTree/Source/AnimNextStateTree/Public/Conditions/AnimNextStateTreeRigVMConditionBase.h
--- a/Engine/Plugins/Experimental/AnimNextStateTree/Source/AnimNextStateTree/Public/Conditions/AnimNextStateTreeRigVMConditionBase.h
+++ b/Engine/Plugins/Experimental/AnimNextStateTree/Source/AnimNextStateTree/Public/Conditions/AnimNextStateTreeRigVMConditionBase.h
@@ -16,7 +16,10 @@
 class UAnimNextDataInterface;
 class UAnimNextDataInterface_EditorData;
 class URigVMGraph;
+class UStateTreeState;
 
+struct FAnimNextStateTreeProgrammaticGraphParams;
+struct FStateTreeBindableStructDesc;
 struct FStateTreeExecutionContext;
 struct FRigVMClient;
 
@@ -37,7 +40,7 @@ public:
 	bool IsIndexCacheInitialized() const;
 
 	/** True if successfully populated index cache from Data Interface, or already populated*/
-	bool TryPopulateIndexCache(TObjectPtr<UAnimNextDataInterface> InDataInterface, FName InResultName, FName InFunctionName);
+	bool TryPopulateIndexCache(const FAnimNextStateTreeRigVMConditionBase& OwningCondition, TObjectPtr<UAnimNextDataInterface> InDataInterface, FName InResultName, FName InFunctionName);
 
 	/** Gets Argument Indexes, you are responsible for checking if initialized first */
 	TArrayView<const uint8> GetVMArgumentIndexes() const;
@@ -70,8 +73,6 @@ struct ANIMNEXTSTATETREE_API FAnimNextStateTreeRigVMConditionBase : public FStat
 
 	static const FName InternalEventPrefixName;
 	static const FName InternalVariablePrefixName;
-	static FName MakeInternalEventName(const FName InName, const FName InFunctionName);
-	static FName MakeInternalVariableName(const FName InName, const FName InFunctionName);
 
 public:
 
@@ -90,9 +91,15 @@ public:
 	
 #if WITH_EDITOR
 	/** Adds condition graphs to parent StateTree RigVM */
-	virtual void GetProgrammaticGraphs(UAnimNextDataInterface_EditorData* EditorData, FRigVMClient& RigVMClient, TArray<URigVMGraph*>& OutGraphs) const;
+	virtual void GetProgrammaticGraphs(FAnimNextStateTreeProgrammaticGraphParams& InProgrammaticGraphParams, const UStateTreeState* State, const FStateTreeBindableStructDesc& Desc);
 #endif
 
+	/** Generates a unique internal name for RigVM events, name not intended to be user-facing */
+	FName MakeInternalEventName(const FName InName) const;
+
+	/** Generates a unique internal name for RigVM variables, name not intended to be user-facing */
+	FName MakeInternalVariableName(const FName InName) const;
+
 public:
 
 #if WITH_EDITORONLY_DATA
@@ -107,6 +114,14 @@ public:
 	UPROPERTY()
 	FRigVMGraphFunctionHeader RigVMFunctionHeader = FRigVMGraphFunctionHeader();
 
+	/** Owning state name, populated during programmatic graph creation */
+	UPROPERTY()
+	FName StateName = NAME_None;
+
+	/** External node ID defined by owning state tree, populated during programmatic graph creation */
+	UPROPERTY()
+	FGuid NodeId = FGuid();
+
 	UPROPERTY()
 	FName InternalEventName = NAME_None;
 
diff --git a/Engine/Plugins/Experimental/AnimNextStateTree/Source/AnimNextStateTreeUncookedOnly/Internal/AnimNextStateTree_EditorData.cpp b/Engine/Plugins/Experimental/AnimNextStateTree/Source/AnimNextStateTreeUncookedOnly/Internal/AnimNextStateTree_EditorData.cpp
--- a/Engine/Plugins/Experimental/AnimNextStateTree/Source/AnimNextStateTreeUncookedOnly/Internal/AnimNextStateTree_EditorData.cpp
+++ b/Engine/Plugins/Experimental/AnimNextStateTree/Source/AnimNextStateTreeUncookedOnly/Internal/AnimNextStateTree_EditorData.cpp
@@ -4,8 +4,10 @@
 
 #include "AnimGraphUncookedOnlyUtils.h"
 #include "AnimNextStateTree.h"
+#include "AnimNextStateTreeEditorOnlyTypes.h"
 #include "AnimNextStateTreeWorkspaceAssetUserData.h"
 #include "Conditions/AnimNextStateTreeRigVMConditionBase.h"
+#include "DataInterface/AnimNextDataInterface_EditorData.h"
 #include "StateTree.h"
 #include "StateTreeEditorData.h"
 #include "Entries/AnimNextVariableEntry.h"
@@ -14,6 +16,9 @@
 #include "Graph/RigUnit_AnimNextGraphRoot.h"
 #include "Graph/RigDecorator_AnimNextCppTrait.h"
 #include "Graph/RigUnit_AnimNextTraitStack.h"
+#include "RigVMCompiler/RigVMCompiler.h"
+#include "RigVMModel/RigVMClient.h"
+#include "RigVMModel/RigVMGraph.h"
 #include "TraitCore/TraitRegistry.h"
 #include "TraitCore/Trait.h"
 #include "AnimStateTreeTrait.h"
@@ -127,12 +132,41 @@ void UAnimNextStateTree_EditorData::GetProgrammaticGraphs(const FRigVMCompileSet
 		}
 
 		// Give child nodes a chance to add graphs
-		for (FConstStructView Node : AnimStateTree->StateTree->GetNodes())
+		FAnimNextStateTreeProgrammaticGraphParams ProgrammaticGraphParams(this, InSettings, RigVMClient, OutGraphs);
+
+		// Populate UID info on our nodes
+		auto MakeNodeGraphs = [&](UStateTreeState& State, UStateTreeState* ParentState)
+		{
+			for (const FStateTreeEditorNode& Node : State.EnterConditions)
+			{
+				// Manually check for each child node type, UStruct's do not reflect interfaces so we can't use an interface here.
+				if (const FAnimNextStateTreeRigVMConditionBase* Condition = Node.Node.GetPtr<FAnimNextStateTreeRigVMConditionBase>())
+				{
+					FStateTreeBindableStructDesc Desc;
+					Desc.StatePath = State.GetPath() + TEXT("/EnterConditions");
+					Desc.Struct = Condition->GetInstanceDataType();
+					Desc.Name = Node.GetName();
+					Desc.ID = Node.ID;
+					Desc.DataSource = EStateTreeBindableStructSource::Condition;
+
+					// Graph gen isn't const since we need to cache generated names in the node
+					const_cast<FAnimNextStateTreeRigVMConditionBase*>(Condition)->GetProgrammaticGraphs(ProgrammaticGraphParams, &State, Desc);
+				}
+			}
+
+			// @TODO: Add RigVM graph task support
+			//for (const FStateTreeEditorNode& Node : State.Tasks)
+			//{
+			//}
+
+			return EStateTreeVisitor::Continue;
+		};
+
+		if (AnimStateTree->StateTree)
 		{
-			// Manually check for each child node type, UStruct's do not reflect interfaces so we can't use an interface here.
-			if (const FAnimNextStateTreeRigVMConditionBase* Condition = Node.GetPtr<const FAnimNextStateTreeRigVMConditionBase>())
+			if (UStateTreeEditorData* StateTreeEditorData = Cast<UStateTreeEditorData>(AnimStateTree->StateTree->EditorData))
 			{
-				Condition->GetProgrammaticGraphs(this, RigVMClient, OutGraphs);
+				StateTreeEditorData->VisitHierarchy(MakeNodeGraphs);
 			}
 		}
 
