Fix StateTree RigVM Conditions not working multithreaded due to reuse of execution contexts.

diff --git a/Engine/Plugins/Experimental/AnimNextStateTree/Source/AnimNextStateTree/Private/Conditions/AnimNextStateTreeRigVMConditionBase.cpp b/Engine/Plugins/Experimental/AnimNextStateTree/Source/AnimNextStateTree/Private/Conditions/AnimNextStateTreeRigVMConditionBase.cpp
--- a/Engine/Plugins/Experimental/AnimNextStateTree/Source/AnimNextStateTree/Private/Conditions/AnimNextStateTreeRigVMConditionBase.cpp
+++ b/Engine/Plugins/Experimental/AnimNextStateTree/Source/AnimNextStateTree/Private/Conditions/AnimNextStateTreeRigVMConditionBase.cpp
@@ -1,6 +1,7 @@
 // Copyright Epic Games, Inc. All Rights Reserved.
 
 #include "Conditions/AnimNextStateTreeRigVMConditionBase.h"
+#include "Graph/AnimNextGraphContextData.h"
 #include "Graph/AnimNextGraphInstance.h"
 #include "StateTreeLinker.h"
 #include "StateTreeExecutionContext.h"
@@ -171,25 +172,34 @@ bool FAnimNextStateTreeRigVMConditionBase::TestCondition(FStateTreeExecutionCont
 			}
 
 			// Prior to execution, copy over ST binding values to RigVM variables
+			FAnimNextGraphInstance& GraphInstance = AnimExecContext->GetRootGraphInstance();
 			TArrayView<const uint8> ArgumentIndexes = InstanceData->GetVMArgumentIndexes();
-			TConstArrayView<FPropertyBagPropertyDesc> ArgumentPropertyDescs = InstanceData->ParamData.GetPropertyBagStruct()->GetPropertyDescs();
-			TArray<FRigVMExternalVariable> ExternalVariables = RigVMInstance->GetExternalVariables();
+			TConstArrayView<FPropertyBagPropertyDesc> StateTreePropertyDescs = InstanceData->ParamData.GetPropertyBagStruct()->GetPropertyDescs();
+			TConstArrayView<FPropertyBagPropertyDesc> AnimNextGraphPropertyDescs = GraphInstance.GetVariables().GetPropertyBagStruct()->GetPropertyDescs();
 			FStructView StateTreeParamData = InstanceData->ParamData.GetMutableValue();
 
 			for (int32 Index = 0; Index < ArgumentIndexes.Num(); Index++)
 			{
-				const FPropertyBagPropertyDesc& ArgumentDesc = ArgumentPropertyDescs[Index];
 				uint8 RigVMExternalVariableIndex = ArgumentIndexes[Index];
-				void* SourceAddress = ArgumentDesc.CachedProperty->ContainerPtrToValuePtr<void*>(StateTreeParamData.GetMemory());
-				void* TargetAddress = ExternalVariables[RigVMExternalVariableIndex].Memory;
 
-				ArgumentDesc.CachedProperty->CopyCompleteValue(TargetAddress, SourceAddress);
+				const FPropertyBagPropertyDesc& StateTreePropertyDesc = StateTreePropertyDescs[Index];
+				void* SourceAddress = StateTreePropertyDesc.CachedProperty->ContainerPtrToValuePtr<void*>(StateTreeParamData.GetMemory());
+
+				const FPropertyBagPropertyDesc& AnimNextGraphDesc = AnimNextGraphPropertyDescs[RigVMExternalVariableIndex];
+				GraphInstance.GetVariables().SetValue(AnimNextGraphDesc.Name, StateTreePropertyDesc.CachedProperty, StateTreeParamData.GetMemory());
 			}
 
-			RigVMInstance->GetVM()->ExecuteVM(RigVMInstance->GetRigVMExtendedExecuteContext(), InternalEventName);
-			bool Result = ExternalVariables[InstanceData->GetVMResultIndex()].GetValue<bool>();
+			FRigVMExtendedExecuteContext& ExtendedExecuteContext = GraphInstance.GetExtendedExecuteContext();
+			FAnimNextExecuteContext& AnimNextContext = GraphInstance.GetExtendedExecuteContext().GetPublicDataSafe<FAnimNextExecuteContext>();
+
+			FAnimNextGraphContextData ContextData(GraphInstance.GetModuleInstance(), &GraphInstance);
+			UE::AnimNext::FScopedExecuteContextData ContextDataScope(AnimNextContext, ContextData);
+
+			bool bResult = false;
+			RigVMInstance->GetVM()->ExecuteVM(ExtendedExecuteContext, InternalEventName);
+			GraphInstance.GetVariable(InternalResultName, bResult);
 
-			return Result;
+			return bResult;
 		}
 	}
 
