Fix CIS Issue 715116: missing dependency

diff --git a/Engine/Plugins/Experimental/AnimNextStateTree/Source/AnimNextStateTreeUncookedOnly/AnimNextStateTreeUncookedOnly.Build.cs b/Engine/Plugins/Experimental/AnimNextStateTree/Source/AnimNextStateTreeUncookedOnly/AnimNextStateTreeUncookedOnly.Build.cs
--- a/Engine/Plugins/Experimental/AnimNextStateTree/Source/AnimNextStateTreeUncookedOnly/AnimNextStateTreeUncookedOnly.Build.cs
+++ b/Engine/Plugins/Experimental/AnimNextStateTree/Source/AnimNextStateTreeUncookedOnly/AnimNextStateTreeUncookedOnly.Build.cs
@@ -9,6 +9,7 @@ namespace UnrealBuildTool.Rules
 			PrivateDependencyModuleNames.AddRange(new string[] 
 			{
 				"StateTreeEditorModule",
+				"PropertyBindingUtils",
 				"RigVM",
 				"RigVMDeveloper"
 			});
