Fix missing PropertyBindingUtils dependency in AnimNextStateTreeEditor only showing on Mac

#tests Test with and without the fix in two preflights

diff --git a/Engine/Plugins/Experimental/AnimNextStateTree/Source/AnimNextStateTreeEditor/AnimNextStateTreeEditor.Build.cs b/Engine/Plugins/Experimental/AnimNextStateTree/Source/AnimNextStateTreeEditor/AnimNextStateTreeEditor.Build.cs
--- a/Engine/Plugins/Experimental/AnimNextStateTree/Source/AnimNextStateTreeEditor/AnimNextStateTreeEditor.Build.cs
+++ b/Engine/Plugins/Experimental/AnimNextStateTree/Source/AnimNextStateTreeEditor/AnimNextStateTreeEditor.Build.cs
@@ -32,6 +32,7 @@ namespace UnrealBuildTool.Rules
 					"WorkspaceEditor",
 					"StateTreeEditorModule",
 					"StateTreeModule",
+					"PropertyBindingUtils",
 					"AnimNextStateTreeUncookedOnly",
 					"AnimNextEditor",
 					"SlateCore",
