AnimNext ST debugging parity with ST
- Fix debugger breaking after changing workspace selection
- Fix debugger clearing debugger when selecting within same state tree
- Properly name ST for debugger & also update previous names via postload
- Set AN ST debug context object to something unique per ST, so trace is unique
- Fix issue where ST would not trace if there was no ST running on game thread.

#tests PIE

diff --git a/Engine/Plugins/Experimental/AnimNextStateTree/Source/AnimNextStateTree/Private/AnimNextStateTree.cpp b/Engine/Plugins/Experimental/AnimNextStateTree/Source/AnimNextStateTree/Private/AnimNextStateTree.cpp
new file mode 100644
--- /dev/null
+++ b/Engine/Plugins/Experimental/AnimNextStateTree/Source/AnimNextStateTree/Private/AnimNextStateTree.cpp
@@ -0,0 +1,24 @@
+// Copyright Epic Games, Inc. All Rights Reserved.
+
+#include "AnimNextStateTree.h"
+
+#include "Serialization/CustomVersion.h"
+
+const FGuid FAnimNextStateTreeCustomVersion::GUID(0x45641511, 0x102F42BB, 0xA6EF181D, 0x6C442CAC);
+FCustomVersionRegistration GRegisterStateTreeCustomVersion(FAnimNextStateTreeCustomVersion::GUID, FAnimNextStateTreeCustomVersion::LatestVersion, TEXT("AnimNextStateTree"));
+
+void UAnimNextStateTree::PostLoad()
+{
+	Super::PostLoad();
+
+#if WITH_EDITORONLY_DATA
+	const int32 CurrentVersion = GetLinkerCustomVersion(FAnimNextStateTreeCustomVersion::GUID);
+	if (CurrentVersion < FAnimNextStateTreeCustomVersion::InnerStateTreeUniqueName)
+	{
+		if (StateTree && StateTree->GetFName() == "StateTree")
+		{
+			StateTree->Rename(*(GetFName().ToString() + "_StateTree"));
+		}
+	}
+#endif // WITH_EDITORONLY_DATA
+}
\ No newline at end of file
diff --git a/Engine/Plugins/Experimental/AnimNextStateTree/Source/AnimNextStateTree/Private/AnimStateTreeTrait.cpp b/Engine/Plugins/Experimental/AnimNextStateTree/Source/AnimNextStateTree/Private/AnimStateTreeTrait.cpp
--- a/Engine/Plugins/Experimental/AnimNextStateTree/Source/AnimNextStateTree/Private/AnimStateTreeTrait.cpp
+++ b/Engine/Plugins/Experimental/AnimNextStateTree/Source/AnimNextStateTree/Private/AnimStateTreeTrait.cpp
@@ -9,6 +9,10 @@
 #include "Logging/LogScopedVerbosityOverride.h"
 #include "TraitCore/NodeInstance.h"
 
+#if ENABLE_ANIM_DEBUG 
+#include "Debugger/StateTreeRuntimeValidation.h"
+#endif // ENABLE_ANIM_DEBUG 
+
 namespace UE::AnimNext
 {
 	AUTO_REGISTER_ANIM_TRAIT(FStateTreeTrait)
@@ -30,6 +34,43 @@ namespace UE::AnimNext
 
 	void FStateTreeTrait::FInstanceData::Destruct(const FExecutionContext& Context, const FTraitBinding& Binding)
 	{
+		// @TODO: UE-240683
+		// We should call stop, but at the moment we don't have a good time / way to do that 
+		
+//		// Use Transient as owner, since our owning component is already dis-associated by now and all we need is the ST + Instance to stop.
+//		UObject* Owner = GetTransientPackage();
+//
+//		if (StateTree && false)
+//		{
+//			FStateTreeExecutionContext StateTreeExecutionContext(*Owner, *StateTree, InstanceData);
+//#if ENABLE_ANIM_DEBUG 
+//			// Clear owner, since at this point we are destructing and don't care about the context.
+//			// Also since the real owner may be GC'd by now (shutdown, level swap, etc).
+//			InstanceData.GetRuntimeValidation().SetContext(nullptr, StateTree);
+//#endif // ENABLE_ANIM_DEBUG 
+//			FAnimNextStateTreeTraitContext TraitContext(const_cast<FExecutionContext&>(Context), &Binding);
+//
+//			// @TODO: In separate CL make a helper method to do this. 
+//			StateTreeExecutionContext.SetCollectExternalDataCallback(FOnCollectStateTreeExternalData::CreateLambda([this, &TraitContext]( const FStateTreeExecutionContext& Context, const UStateTree* StateTree, TArrayView<const FStateTreeExternalDataDesc> ExternalDataDescs, TArrayView<FStateTreeDataView> OutDataViews) -> bool
+//			{
+//				for (int32 Index = 0; Index < ExternalDataDescs.Num(); Index++)
+//				{
+//					const FStateTreeExternalDataDesc& ItemDesc = ExternalDataDescs[Index];
+//					if (ItemDesc.Struct != nullptr)
+//					{
+//						if (ItemDesc.Struct->IsChildOf(FAnimNextStateTreeTraitContext::StaticStruct()))
+//						{
+//							OutDataViews[Index] = FStateTreeDataView(FAnimNextStateTreeTraitContext::StaticStruct(), reinterpret_cast<uint8*>(&TraitContext));
+//						}
+//					}
+//				}
+//
+//				return true;
+//			}));
+//
+//			StateTreeExecutionContext.Stop();
+//		}
+
 		FTrait::FInstanceData::Destruct(Context, Binding);
 		IGarbageCollection::UnregisterWithGC(Context, Binding);
 	}
@@ -43,6 +84,17 @@ namespace UE::AnimNext
 		if (InstanceData->StateTree)
 		{
 			UObject* Owner = GetTransientPackage();
+#if ENABLE_ANIM_DEBUG 
+			Owner = const_cast<UObject*>(Context.GetHostObject());
+			// @TODO: Makes unique, but breaks visual logger
+			Owner = Owner ? Owner->GetOuter() : Owner;
+
+			if (!Owner)
+			{
+				return;
+			}
+#endif // ENABLE_ANIM_DEBUG 
+
 			FStateTreeExecutionContext StateTreeExecutionContext(*Owner, *InstanceData->StateTree, InstanceData->InstanceData);
 
 			FAnimNextStateTreeTraitContext TraitContext(Context, &Binding);
@@ -102,6 +154,16 @@ namespace UE::AnimNext
 		if (InstanceData->StateTree)
 		{
 			UObject* Owner = GetTransientPackage();
+#if ENABLE_ANIM_DEBUG 
+			Owner = const_cast<UObject*>(Context.GetHostObject());
+			// @TODO: Makes unique, but breaks visual logger
+			Owner = Owner ? Owner->GetOuter() : Owner;
+
+			if (!Owner)
+			{
+				return;
+			}
+#endif // ENABLE_ANIM_DEBUG 
 			
 			FStateTreeExecutionContext StateTreeExecutionContext(*Owner, *InstanceData->StateTree, InstanceData->InstanceData);
 			FAnimNextGraphInstance* OwnerGraphInstance = &Binding.GetTraitPtr().GetNodeInstance()->GetOwner();
diff --git a/Engine/Plugins/Experimental/AnimNextStateTree/Source/AnimNextStateTree/Public/AnimNextStateTree.h b/Engine/Plugins/Experimental/AnimNextStateTree/Source/AnimNextStateTree/Public/AnimNextStateTree.h
--- a/Engine/Plugins/Experimental/AnimNextStateTree/Source/AnimNextStateTree/Public/AnimNextStateTree.h
+++ b/Engine/Plugins/Experimental/AnimNextStateTree/Source/AnimNextStateTree/Public/AnimNextStateTree.h
@@ -13,6 +13,26 @@ namespace UE::AnimNext::UncookedOnly
 struct FUtils;
 }
 
+struct ANIMNEXTSTATETREE_API FAnimNextStateTreeCustomVersion
+{
+	enum Type
+	{
+		// Before any version changes were made in the plugin
+		BeforeCustomVersionWasAdded = 0,
+		// Added unique name to inner state tree
+		InnerStateTreeUniqueName,
+
+		// -----<new versions can be added above this line>-------------------------------------------------
+		VersionPlusOne,
+		LatestVersion = VersionPlusOne - 1
+	};
+
+	/** The GUID for this custom version number */
+	const static FGuid GUID;
+
+private:
+	FAnimNextStateTreeCustomVersion() = default;
+};
 
 UCLASS(BlueprintType)
 class ANIMNEXTSTATETREE_API UAnimNextStateTree : public UAnimNextAnimationGraph
@@ -25,4 +45,10 @@ public:
 
 	UPROPERTY()
 	TObjectPtr<UStateTree> StateTree;
+
+public:
+
+	//~ Begin UObject Interface
+	virtual void PostLoad() override;
+	//~ End UObject Interface
 };
\ No newline at end of file
diff --git a/Engine/Plugins/Experimental/AnimNextStateTree/Source/AnimNextStateTree/Public/AnimNextStateTreeContext.h b/Engine/Plugins/Experimental/AnimNextStateTree/Source/AnimNextStateTree/Public/AnimNextStateTreeContext.h
--- a/Engine/Plugins/Experimental/AnimNextStateTree/Source/AnimNextStateTree/Public/AnimNextStateTreeContext.h
+++ b/Engine/Plugins/Experimental/AnimNextStateTree/Source/AnimNextStateTree/Public/AnimNextStateTreeContext.h
@@ -30,7 +30,7 @@ struct FAnimNextStateTreeTraitContext
 	float QueryPlaybackRatio(TNonNullPtr<UObject> InAsset) const;
 
 	UE::AnimNext::FExecutionContext* GetAnimExecuteContext() { return Context; }
-protected:
+public:
 	FAnimNextStateTreeTraitContext(UE::AnimNext::FExecutionContext& InContext, const UE::AnimNext::FTraitBinding* InBinding) : Context(&InContext), Binding(InBinding) {}
 protected:
 	UE::AnimNext::FExecutionContext* Context = nullptr;
diff --git a/Engine/Plugins/Experimental/AnimNextStateTree/Source/AnimNextStateTreeEditor/Private/AnimNextStateTreeEditorHost.cpp b/Engine/Plugins/Experimental/AnimNextStateTree/Source/AnimNextStateTreeEditor/Private/AnimNextStateTreeEditorHost.cpp
--- a/Engine/Plugins/Experimental/AnimNextStateTree/Source/AnimNextStateTreeEditor/Private/AnimNextStateTreeEditorHost.cpp
+++ b/Engine/Plugins/Experimental/AnimNextStateTree/Source/AnimNextStateTreeEditor/Private/AnimNextStateTreeEditorHost.cpp
@@ -45,7 +45,22 @@ bool FAnimNextStateTreeEditorHost::ShouldShowCompileButton() const
 
 void FAnimNextStateTreeEditorHost::OnWorkspaceFocussedDocumentChanged(TObjectPtr<UObject> InObject) const
 {
-	OnStateTreeChangedDelegate.Broadcast();
+	TObjectPtr<UAnimNextStateTree> InStateTree = Cast<UAnimNextStateTree>(InObject);
+
+	if (TStrongObjectPtr<UAnimNextStateTree> LastStateTreePinned = WeakLastStateTree.Pin())
+	{
+		if (LastStateTreePinned != InStateTree)
+		{
+			OnStateTreeChangedDelegate.Broadcast();
+		}
+	}
+	else if (InStateTree)
+	{
+		OnStateTreeChangedDelegate.Broadcast();
+	}
+
+	// Always set the last state tree as incoming object after cast. If it's null that will clear the last state tree
+	WeakLastStateTree = InStateTree;
 }
 
 FSimpleMulticastDelegate& FAnimNextStateTreeEditorHost::OnStateTreeChanged()
diff --git a/Engine/Plugins/Experimental/AnimNextStateTree/Source/AnimNextStateTreeEditor/Private/AnimNextStateTreeEditorHost.h b/Engine/Plugins/Experimental/AnimNextStateTree/Source/AnimNextStateTreeEditor/Private/AnimNextStateTreeEditorHost.h
--- a/Engine/Plugins/Experimental/AnimNextStateTree/Source/AnimNextStateTreeEditor/Private/AnimNextStateTreeEditorHost.h
+++ b/Engine/Plugins/Experimental/AnimNextStateTree/Source/AnimNextStateTreeEditor/Private/AnimNextStateTreeEditorHost.h
@@ -9,6 +9,8 @@ namespace UE::Workspace
 	class IWorkspaceEditor;
 }
 
+class UAnimNextStateTree;
+
 class FAnimNextStateTreeEditorHost : public IStateTreeEditorHost
 {
 public:
@@ -28,4 +30,7 @@ protected:
 
 	TWeakPtr<UE::Workspace::IWorkspaceEditor> WeakWorkspaceEditor;	
 	FSimpleMulticastDelegate OnStateTreeChangedDelegate;
+
+	/** Last state tree, cached to avoid changing debugger & losing context more than needed */
+	mutable TWeakObjectPtr<UAnimNextStateTree> WeakLastStateTree;
 };
\ No newline at end of file
diff --git a/Engine/Plugins/Experimental/AnimNextStateTree/Source/AnimNextStateTreeEditor/Private/AnimNextStateTreeFactory.cpp b/Engine/Plugins/Experimental/AnimNextStateTree/Source/AnimNextStateTreeEditor/Private/AnimNextStateTreeFactory.cpp
--- a/Engine/Plugins/Experimental/AnimNextStateTree/Source/AnimNextStateTreeEditor/Private/AnimNextStateTreeFactory.cpp
+++ b/Engine/Plugins/Experimental/AnimNextStateTree/Source/AnimNextStateTreeEditor/Private/AnimNextStateTreeFactory.cpp
@@ -32,12 +32,13 @@ UObject* UAnimNextStateTreeFactory::FactoryCreateNew(UClass* Class, UObject* InP
 	// Create internal state tree	
 	UStateTreeFactory* StateTreeFactory = NewObject<UStateTreeFactory>(GetTransientPackage());
 	StateTreeFactory->SetSchemaClass(UStateTreeAnimNextSchema::StaticClass());
-	UStateTree* StateTree = CastChecked<UStateTree>(StateTreeFactory->FactoryCreateNew(UStateTree::StaticClass(), NewStateTree, TEXT("StateTree"), RF_Transactional, nullptr, nullptr));
+
+	UStateTree* StateTree = CastChecked<UStateTree>(StateTreeFactory->FactoryCreateNew(UStateTree::StaticClass(), NewStateTree, FName(Name.ToString() + "_StateTree"), RF_Transactional, nullptr, nullptr));
 
 	NewStateTree->StateTree = StateTree;
 
 	// Create internal editor data
-	UAnimNextStateTree_EditorData* EditorData = NewObject<UAnimNextStateTree_EditorData>(NewStateTree, TEXT("EditorData"), RF_Transactional);
+	UAnimNextStateTree_EditorData* EditorData = NewObject<UAnimNextStateTree_EditorData>(NewStateTree, FName(Name.ToString() + "_StateTree_EditorData"), RF_Transactional);
 	NewStateTree->EditorData = EditorData;
 	EditorData->bUsesExternalPackages = false;
 	EditorData->Initialize(/*bRecompileVM*/false);
diff --git a/Engine/Plugins/Runtime/StateTree/Source/StateTreeEditorModule/Private/Debugger/SStateTreeDebuggerView.cpp b/Engine/Plugins/Runtime/StateTree/Source/StateTreeEditorModule/Private/Debugger/SStateTreeDebuggerView.cpp
--- a/Engine/Plugins/Runtime/StateTree/Source/StateTreeEditorModule/Private/Debugger/SStateTreeDebuggerView.cpp
+++ b/Engine/Plugins/Runtime/StateTree/Source/StateTreeEditorModule/Private/Debugger/SStateTreeDebuggerView.cpp
@@ -20,6 +20,7 @@
 #include "StateTreeDebuggerTrack.h"
 #include "StateTreeDelegates.h"
 #include "StateTreeEditorData.h"
+#include "StateTreeEditorModule.h"
 #include "StateTreeEditorSettings.h"
 #include "StateTreeEditorStyle.h"
 #include "StateTreeModule.h"
@@ -885,6 +886,13 @@ void SStateTreeDebuggerView::OnDebuggerScrubStateChanged(const UE::StateTreeDebu
 	TArray<TSharedPtr<FStateTreeDebuggerEventTreeElement>, TInlineAllocator<8>> ScopeStack;
 
 	const int32 SpanIdx = ScrubState.GetFrameSpanIndex();
+
+	if (!Spans.IsValidIndex(SpanIdx))
+	{
+		UE_LOG(LogStateTreeEditor, Error, TEXT("Invalid index in span: Idx: %i, Num Spans: %i"), SpanIdx, Spans.Num());
+		return;
+	}
+
 	const int32 FirstEventIdx = Spans[SpanIdx].EventIdx;
 	const TraceServices::FFrame Frame = Spans[SpanIdx].Frame;
 	const int32 MaxEventIdx = Spans.IsValidIndex(SpanIdx+1) ? Spans[SpanIdx+1].EventIdx : Events.Num();
diff --git a/Engine/Plugins/Runtime/StateTree/Source/StateTreeEditorModule/Private/StateTreeEditorModeToolkit.cpp b/Engine/Plugins/Runtime/StateTree/Source/StateTreeEditorModule/Private/StateTreeEditorModeToolkit.cpp
--- a/Engine/Plugins/Runtime/StateTree/Source/StateTreeEditorModule/Private/StateTreeEditorModeToolkit.cpp
+++ b/Engine/Plugins/Runtime/StateTree/Source/StateTreeEditorModule/Private/StateTreeEditorModeToolkit.cpp
@@ -601,6 +601,15 @@ void FStateTreeEditorModeToolkit::UpdateStateTreeOutliner()
 void FStateTreeEditorModeToolkit::UpdateDebuggerView()
 {
 	DebuggerView = SNullWidget::NullWidget;
+	TSharedPtr<SDockTab> SharedDebuggerTab = WeakDebuggerTab.Pin();
+
+	if (SharedDebuggerTab)
+	{
+		// Clear any references the previous tab might have to a previous debugger view.
+		// The view will clear the shared debugger's bindings on dtor.
+		// We don't want it to clear newly created bindings from our new view.
+		SharedDebuggerTab->SetContent(SNullWidget::NullWidget);
+	}
 	
 	if (UStateTree* StateTree = WeakEditorMode->GetStateTree())
     {
@@ -610,7 +619,7 @@ void FStateTreeEditorModeToolkit::UpdateDebuggerView()
     	}		
     }
 
-	if(TSharedPtr<SDockTab> SharedDebuggerTab = WeakDebuggerTab.Pin())
+	if (SharedDebuggerTab)
 	{
 		SharedDebuggerTab->SetContent(DebuggerView.ToSharedRef());
 	}
diff --git a/Engine/Plugins/Runtime/StateTree/Source/StateTreeModule/Private/Debugger/StateTreeTrace.cpp b/Engine/Plugins/Runtime/StateTree/Source/StateTreeModule/Private/Debugger/StateTreeTrace.cpp
--- a/Engine/Plugins/Runtime/StateTree/Source/StateTreeModule/Private/Debugger/StateTreeTrace.cpp
+++ b/Engine/Plugins/Runtime/StateTree/Source/StateTreeModule/Private/Debugger/StateTreeTrace.cpp
@@ -122,6 +122,8 @@ UE_TRACE_EVENT_END()
 namespace UE::StateTreeTrace
 {
 
+static double GWorldTime = 0.0;
+
 FDelegateHandle GOnWorldTickStartDelegateHandle;
 FDelegateHandle GTracingStateChangedDelegateHandle;
 
@@ -256,7 +258,6 @@ struct FInstanceEventBufferedData
 /** Struct to keep track of the buffered event data and flush them. */
 struct FBufferedDataList
 {
-	double RecordingWorldTime = -1;
 	double TracedRecordingWorldTime = -1;
 	
 	/**
@@ -293,7 +294,7 @@ struct FBufferedDataList
 		// Trace asset events first since they are required for instance lifetime event types.
 		// Events are preserved in case the trace session is stopped and then a new one gets started
 		// in the same game session. In which case we need to output the ids to that new trace.
-		if (bTraceLifetimeEvents)
+		if (bTraceLifetimeEvents || !IsInGameThread())
 		{
 			for (const FAssetDebugIdEventBufferedData& AssetDebugIdEventData : AssetDebugIdEvents)
 			{
@@ -326,11 +327,11 @@ struct FBufferedDataList
 	 */
 	void TraceWorldTime()
 	{
-		if (TracedRecordingWorldTime != RecordingWorldTime)
+		if (TracedRecordingWorldTime != GWorldTime)
 		{
-			TracedRecordingWorldTime = RecordingWorldTime;
+			TracedRecordingWorldTime = GWorldTime;
 			UE_TRACE_LOG(StateTreeDebugger, WorldTimestampEvent, StateTreeDebugChannel)
-				<< WorldTimestampEvent.WorldTime(RecordingWorldTime);
+				<< WorldTimestampEvent.WorldTime(GWorldTime);
 		}
 	}
 
@@ -379,7 +380,7 @@ struct FBufferedDataList
 		CurrentVersion++;
 
 		// Force world time to trace on the next event
-		RecordingWorldTime = -1;
+		GWorldTime = -1;
 	}
 };
 
@@ -488,10 +489,10 @@ void SerializeDataViewToArchive(FBufferArchive& Ar, const FStateTreeDataView Dat
 
 void RegisterGlobalDelegates()
 {
-	GOnWorldTickStartDelegateHandle = FWorldDelegates::OnWorldTickStart.AddLambda([&WorldTime=GBufferedEvents.RecordingWorldTime](const UWorld* TickedWorld, ELevelTick TickType, float DeltaTime)
+	GOnWorldTickStartDelegateHandle = FWorldDelegates::OnWorldTickStart.AddLambda([](const UWorld* TickedWorld, ELevelTick TickType, float DeltaTime)
 		{
 #if OBJECT_TRACE_ENABLED
-			WorldTime = FObjectTrace::GetWorldElapsedTime(TickedWorld);
+			GWorldTime = FObjectTrace::GetWorldElapsedTime(TickedWorld);
 #endif// OBJECT_TRACE_ENABLED
 		});
 
