Fix link issue introduced with CL 39529849
#rnx

diff --git a/Engine/Plugins/Experimental/AnimNextStateTree/Source/AnimNextStateTree/Private/AnimNextStateTree.cpp b/Engine/Plugins/Experimental/AnimNextStateTree/Source/AnimNextStateTree/Private/AnimNextStateTree.cpp
--- a/Engine/Plugins/Experimental/AnimNextStateTree/Source/AnimNextStateTree/Private/AnimNextStateTree.cpp
+++ b/Engine/Plugins/Experimental/AnimNextStateTree/Source/AnimNextStateTree/Private/AnimNextStateTree.cpp
@@ -5,7 +5,10 @@
 #include "Serialization/CustomVersion.h"
 
 const FGuid FAnimNextStateTreeCustomVersion::GUID(0x45641511, 0x102F42BB, 0xA6EF181D, 0x6C442CAC);
+namespace UE::AnimNextStateTree::Private
+{
 FCustomVersionRegistration GRegisterStateTreeCustomVersion(FAnimNextStateTreeCustomVersion::GUID, FAnimNextStateTreeCustomVersion::LatestVersion, TEXT("AnimNextStateTree"));
+}
 
 void UAnimNextStateTree::PostLoad()
 {
diff --git a/Engine/Plugins/Runtime/StateTree/Source/StateTreeModule/Private/StateTree.cpp b/Engine/Plugins/Runtime/StateTree/Source/StateTreeModule/Private/StateTree.cpp
--- a/Engine/Plugins/Runtime/StateTree/Source/StateTreeModule/Private/StateTree.cpp
+++ b/Engine/Plugins/Runtime/StateTree/Source/StateTreeModule/Private/StateTree.cpp
@@ -27,7 +27,10 @@
 #include UE_INLINE_GENERATED_CPP_BY_NAME(StateTree)
 
 const FGuid FStateTreeCustomVersion::GUID(0x28E21331, 0x501F4723, 0x8110FA64, 0xEA10DA1E);
+namespace UE::StateTree::Private
+{
 FCustomVersionRegistration GRegisterStateTreeCustomVersion(FStateTreeCustomVersion::GUID, FStateTreeCustomVersion::LatestVersion, TEXT("StateTreeAsset"));
+}
 
 bool UStateTree::IsReadyToRun() const
 {
