Prep for graph execute task to AnimNext StateTree
- Add a GetProgrammaticFunctionHeaders step to AnimNext RigVM compilation
- Add support for internal variable setting / getting to AnimNext RigVM
- Extract VariableIndexCache into own struct for reuse.

#tests PIE

diff --git a/Engine/Plugins/Experimental/AnimNext/Source/AnimNextTestSuite/Private/AnimNextDataInterfaceTest.cpp b/Engine/Plugins/Experimental/AnimNext/Source/AnimNextTestSuite/Private/AnimNextDataInterfaceTest.cpp
--- a/Engine/Plugins/Experimental/AnimNext/Source/AnimNextTestSuite/Private/AnimNextDataInterfaceTest.cpp
+++ b/Engine/Plugins/Experimental/AnimNext/Source/AnimNextTestSuite/Private/AnimNextDataInterfaceTest.cpp
@@ -4,6 +4,8 @@
 #include "AnimNextRigVMAssetEditorData.h"
 #include "AnimNextTest.h"
 #include "UncookedOnlyUtils.h"
+#include "Compilation/AnimNextRigVMAssetCompileContext.h"
+#include "Compilation/AnimNextGetVariableCompileContext.h"
 #include "DataInterface/AnimNextDataInterface.h"
 #include "Misc/AutomationTest.h"
 #include "DataInterface/AnimNextDataInterfaceFactory.h"
@@ -37,10 +39,13 @@ bool FDataInterfaceCompile::RunTest(const FString& InParameters)
 	auto CompileAndTest = [A, B, C, D](TFunctionRef<void()> InTestFunction)
 	{
 		FRigVMCompileSettings TempSettings = {};
-		UncookedOnly::FUtils::CompileVariables(TempSettings, A);
-		UncookedOnly::FUtils::CompileVariables(TempSettings, B);
-		UncookedOnly::FUtils::CompileVariables(TempSettings, C);
-		UncookedOnly::FUtils::CompileVariables(TempSettings, D);
+		FAnimNextRigVMAssetCompileContext TempContext = {};
+		FAnimNextGetVariableCompileContext TempGetVariableContext = FAnimNextGetVariableCompileContext(TempContext);
+
+		UncookedOnly::FUtils::CompileVariables(TempSettings, A, TempGetVariableContext);
+		UncookedOnly::FUtils::CompileVariables(TempSettings, B, TempGetVariableContext);
+		UncookedOnly::FUtils::CompileVariables(TempSettings, C, TempGetVariableContext);
+		UncookedOnly::FUtils::CompileVariables(TempSettings, D, TempGetVariableContext);
 
 		InTestFunction();
 	};
diff --git a/Engine/Plugins/Experimental/AnimNext/Source/AnimNextUncookedOnly/Internal/AnimNextRigVMAssetEditorData.h b/Engine/Plugins/Experimental/AnimNext/Source/AnimNextUncookedOnly/Internal/AnimNextRigVMAssetEditorData.h
--- a/Engine/Plugins/Experimental/AnimNext/Source/AnimNextUncookedOnly/Internal/AnimNextRigVMAssetEditorData.h
+++ b/Engine/Plugins/Experimental/AnimNext/Source/AnimNextUncookedOnly/Internal/AnimNextRigVMAssetEditorData.h
@@ -17,6 +17,11 @@ class UAnimNextRigVMAssetEditorData;
 class UAnimNextEdGraph;
 class UAnimNextDataInterfaceEntry;
 class UAnimNextWorkspaceEditorMode;
+struct FAnimNextRigVMAssetCompileContext;
+struct FAnimNextGetFunctionHeaderCompileContext;
+struct FAnimNextGetVariableCompileContext;
+struct FAnimNextGetGraphCompileContext;
+struct FAnimNextProcessGraphCompileContext;
 
 namespace UE::AnimNext::UncookedOnly
 {
@@ -307,24 +312,30 @@ protected:
 	// Override to allow assets to prevent certain entries being created
 	virtual bool CanAddNewEntry(TSubclassOf<UAnimNextRigVMAssetEntry> InClass) const { return true; }
 
+	////////////////////////////////////////////////////////////////////////////////////////////////////////////////////////
 	// Begin Compilation overrides, in order of operation
 
 	// Compilation phase 1: Called before RigVM compilation to setup compiler settings and clean our outer asset of compiler-generated data
 	virtual void OnPreCompileAsset(FRigVMCompileSettings& InSettings) {}
 
-	// Compilation phase 2: Called before RigVM compilation to allow this asset to generate variables to be injected, separate method to allow programmatic graphs to use these vars
+	// Compilation phase 2: Called before RigVM compilation to allow this asset to specify function headers that require generation, along with function generation metadata.
+	// While users may manually generate graphs using function headers, for convience we provide an autogeneration process for function headers requested here.
+	virtual void OnPreCompileGetProgrammaticFunctionHeaders(const FRigVMCompileSettings& InSettings, FAnimNextGetFunctionHeaderCompileContext& OutCompileContext) {}
+
+	// Compilation phase 3: Called before RigVM compilation to allow this asset to generate variables to be injected, separate method to allow programmatic graphs to use these vars
 	// These variables will be regenerated each compile, and are not saved between compiles
-	virtual void OnPreCompileGetProgrammaticVariables(const FRigVMCompileSettings& InSettings, TArray<FAnimNextProgrammaticVariable>& OutVariables) {}
+	virtual void OnPreCompileGetProgrammaticVariables(const FRigVMCompileSettings& InSettings, FAnimNextGetVariableCompileContext& OutCompileContext) {}
 
-	// Compilation phase 3: Called before RigVM compilation to allow this asset to generate graphs to be injected
-	virtual void OnPreCompileGetProgrammaticGraphs(const FRigVMCompileSettings& InSettings, TArray<URigVMGraph*>& OutGraphs) {}
+	// Compilation phase 4: Called before RigVM compilation to allow this asset to generate graphs to be injected
+	virtual void OnPreCompileGetProgrammaticGraphs(const FRigVMCompileSettings& InSettings, FAnimNextGetGraphCompileContext& OutCompileContext) {}
 
-	// Compilation phase 4: Called before RigVM compilation to allow this asset to process, transform or replace the graphs that will be compiled
-	virtual void OnPreCompileProcessGraphs(const FRigVMCompileSettings& InSettings, TArray<URigVMGraph*>& InOutGraphs) {}
+	// Compilation phase 5: Called before RigVM compilation to allow this asset to process, transform or replace the graphs that will be compiled
+	virtual void OnPreCompileProcessGraphs(const FRigVMCompileSettings& InSettings, FAnimNextProcessGraphCompileContext& OutCompileContext) {}
 
-	// Compilation phase 5: Called after RigVM compilation to clean up/finish the compilation process
+	// Compilation phase 6: Called after RigVM compilation to clean up/finish the compilation process
 	virtual void OnPostCompileCleanup(const FRigVMCompileSettings& InSettings) {}
 
+	////////////////////////////////////////////////////////////////////////////////////////////////////////////////////////
 	// End Compilation overrides
 
 	// Customization point for derived types to transform new asset entries
diff --git a/Engine/Plugins/Experimental/AnimNext/Source/AnimNextUncookedOnly/Internal/Compilation/AnimNextGetFunctionHeaderCompileContext.h b/Engine/Plugins/Experimental/AnimNext/Source/AnimNextUncookedOnly/Internal/Compilation/AnimNextGetFunctionHeaderCompileContext.h
new file mode 100644
--- /dev/null
+++ b/Engine/Plugins/Experimental/AnimNext/Source/AnimNextUncookedOnly/Internal/Compilation/AnimNextGetFunctionHeaderCompileContext.h
@@ -0,0 +1,23 @@
+// Copyright Epic Games, Inc. All Rights Reserved.
+
+#pragma once
+
+#include "Compilation/AnimNextRigVMAssetCompileContext.h"
+
+/**
+ * Struct holding temporary compilation info during function header population 
+ */
+struct ANIMNEXTUNCOOKEDONLY_API FAnimNextGetFunctionHeaderCompileContext
+{
+	FAnimNextGetFunctionHeaderCompileContext(FAnimNextRigVMAssetCompileContext& InCompilerContext);
+
+public:
+	
+	const TArray<FAnimNextProgrammaticFunctionHeader>& GetFunctionHeaders() const;
+	TArray<FAnimNextProgrammaticFunctionHeader>& GetMutableFunctionHeaders();
+
+protected:
+
+	FAnimNextRigVMAssetCompileContext& CompilerContext;
+};
+
diff --git a/Engine/Plugins/Experimental/AnimNext/Source/AnimNextUncookedOnly/Internal/Compilation/AnimNextGetGraphCompileContext.h b/Engine/Plugins/Experimental/AnimNext/Source/AnimNextUncookedOnly/Internal/Compilation/AnimNextGetGraphCompileContext.h
new file mode 100644
--- /dev/null
+++ b/Engine/Plugins/Experimental/AnimNext/Source/AnimNextUncookedOnly/Internal/Compilation/AnimNextGetGraphCompileContext.h
@@ -0,0 +1,25 @@
+// Copyright Epic Games, Inc. All Rights Reserved.
+
+#pragma once
+
+#include "Compilation/AnimNextRigVMAssetCompileContext.h"
+
+/**
+ * Struct holding temporary compilation info during function header population
+ */
+struct ANIMNEXTUNCOOKEDONLY_API FAnimNextGetGraphCompileContext
+{
+	FAnimNextGetGraphCompileContext(FAnimNextRigVMAssetCompileContext& InCompilerContext);
+
+public:
+
+	const TArray<FAnimNextProgrammaticFunctionHeader>& GetFunctionHeaders() const;
+
+	const TArray<URigVMGraph*>& GetProgrammaticGraphs() const;
+	TArray<URigVMGraph*>& GetMutableProgrammaticGraphs();
+
+protected:
+
+	FAnimNextRigVMAssetCompileContext& CompilerContext;
+};
+
diff --git a/Engine/Plugins/Experimental/AnimNext/Source/AnimNextUncookedOnly/Internal/Compilation/AnimNextGetVariableCompileContext.h b/Engine/Plugins/Experimental/AnimNext/Source/AnimNextUncookedOnly/Internal/Compilation/AnimNextGetVariableCompileContext.h
new file mode 100644
--- /dev/null
+++ b/Engine/Plugins/Experimental/AnimNext/Source/AnimNextUncookedOnly/Internal/Compilation/AnimNextGetVariableCompileContext.h
@@ -0,0 +1,25 @@
+// Copyright Epic Games, Inc. All Rights Reserved.
+
+#pragma once
+
+#include "Compilation/AnimNextRigVMAssetCompileContext.h"
+
+/**
+ * Struct holding temporary compilation info during function header population
+ */
+struct ANIMNEXTUNCOOKEDONLY_API FAnimNextGetVariableCompileContext
+{
+	FAnimNextGetVariableCompileContext(FAnimNextRigVMAssetCompileContext& InCompilerContext);
+
+public:
+
+	const TArray<FAnimNextProgrammaticFunctionHeader>& GetFunctionHeaders() const;
+
+	const TArray<FAnimNextProgrammaticVariable>& GetProgrammaticVariables() const;
+	TArray<FAnimNextProgrammaticVariable>& GetMutableProgrammaticVariables();
+
+protected:
+
+	FAnimNextRigVMAssetCompileContext& CompilerContext;
+};
+
diff --git a/Engine/Plugins/Experimental/AnimNext/Source/AnimNextUncookedOnly/Internal/Compilation/AnimNextProcessGraphCompileContext.h b/Engine/Plugins/Experimental/AnimNext/Source/AnimNextUncookedOnly/Internal/Compilation/AnimNextProcessGraphCompileContext.h
new file mode 100644
--- /dev/null
+++ b/Engine/Plugins/Experimental/AnimNext/Source/AnimNextUncookedOnly/Internal/Compilation/AnimNextProcessGraphCompileContext.h
@@ -0,0 +1,23 @@
+// Copyright Epic Games, Inc. All Rights Reserved.
+
+#pragma once
+
+#include "Compilation/AnimNextRigVMAssetCompileContext.h"
+
+/**
+ * Struct holding temporary compilation info during function header population
+ */
+struct ANIMNEXTUNCOOKEDONLY_API FAnimNextProcessGraphCompileContext
+{
+	FAnimNextProcessGraphCompileContext(FAnimNextRigVMAssetCompileContext& InCompilerContext);
+
+public:
+
+	const TArray<URigVMGraph*>& GetAllGraphs() const;
+	TArray<URigVMGraph*>& GetMutableAllGraphs();
+
+protected:
+
+	FAnimNextRigVMAssetCompileContext& CompilerContext;
+};
+
diff --git a/Engine/Plugins/Experimental/AnimNext/Source/AnimNextUncookedOnly/Internal/Compilation/AnimNextProgrammaticFunctionHeader.h b/Engine/Plugins/Experimental/AnimNext/Source/AnimNextUncookedOnly/Internal/Compilation/AnimNextProgrammaticFunctionHeader.h
new file mode 100644
--- /dev/null
+++ b/Engine/Plugins/Experimental/AnimNext/Source/AnimNextUncookedOnly/Internal/Compilation/AnimNextProgrammaticFunctionHeader.h
@@ -0,0 +1,32 @@
+// Copyright Epic Games, Inc. All Rights Reserved.
+
+#pragma once
+
+#include "CoreMinimal.h"
+#include "Param/ParamType.h"
+#include "StructUtils/PropertyBag.h"
+#include "RigVMCore/RigVMGraphFunctionDefinition.h"
+
+/**
+ * Struct wrapping a RigVM Function header. Also contains some metadata used to drive compilation.
+ */
+struct ANIMNEXTUNCOOKEDONLY_API FAnimNextProgrammaticFunctionHeader
+{
+	/** RigVM Function Header we are wrapping */
+	FRigVMGraphFunctionHeader Wrapped;
+
+	/** 
+	 * True instructs compiler that we should generate & link programmatic variables for this function's params
+	 * 
+	 * Used by those who cannot set params by linking into the function node during graph compile time (Ex: StateTree).
+	 */
+	bool bGenerateParamVariables = false;
+
+	/** 
+	 * True instructs compiler that we should generate & link programmatic variables for this function's returns  
+	 * 
+	 * Used by those who cannot get returns by linking into the function node during graph compile time (Ex: StateTree).
+	 */
+	bool bGenerateReturnVariables = false;
+};
+
diff --git a/Engine/Plugins/Experimental/AnimNext/Source/AnimNextUncookedOnly/Internal/Compilation/AnimNextRigVMAssetCompileContext.h b/Engine/Plugins/Experimental/AnimNext/Source/AnimNextUncookedOnly/Internal/Compilation/AnimNextRigVMAssetCompileContext.h
new file mode 100644
--- /dev/null
+++ b/Engine/Plugins/Experimental/AnimNext/Source/AnimNextUncookedOnly/Internal/Compilation/AnimNextRigVMAssetCompileContext.h
@@ -0,0 +1,43 @@
+// Copyright Epic Games, Inc. All Rights Reserved.
+
+#pragma once
+
+#include "CoreMinimal.h"
+#include "Param/ParamType.h"
+#include "StructUtils/PropertyBag.h"
+#include "RigVMCore/RigVMGraphFunctionDefinition.h"
+#include "Compilation/AnimNextProgrammaticFunctionHeader.h"
+#include "Variables/AnimNextProgrammaticVariable.h"
+
+class URigVMGraph;
+
+/**
+ * Struct holding temporary compilation info for a single RigVM asset
+ * 
+ * Note: Sub compile contexts give limited access. But feel free to expand
+ * the access of any particular sub context as needed.
+ */
+struct ANIMNEXTUNCOOKEDONLY_API FAnimNextRigVMAssetCompileContext
+{
+private:
+	friend class UAnimNextRigVMAssetEditorData;
+	friend struct FAnimNextGetFunctionHeaderCompileContext;
+	friend struct FAnimNextGetVariableCompileContext;
+	friend struct FAnimNextGetGraphCompileContext;
+	friend struct FAnimNextProcessGraphCompileContext;
+
+protected:
+
+	/** Function Headers we should generate variables / graphs for */
+	TArray<FAnimNextProgrammaticFunctionHeader> FunctionHeaders;
+
+	/** Programmtic variables generated during this compile */
+	TArray<FAnimNextProgrammaticVariable> ProgrammaticVariables;
+
+	/** Programmtic graphs generated during this compile */
+	TArray<URigVMGraph*> ProgrammaticGraphs;
+
+	/** All graphs compiled by this RigVM Asset */
+	TArray<URigVMGraph*> AllGraphs;
+};
+
diff --git a/Engine/Plugins/Experimental/AnimNext/Source/AnimNextUncookedOnly/Internal/Module/AnimNextModule_EditorData.h b/Engine/Plugins/Experimental/AnimNext/Source/AnimNextUncookedOnly/Internal/Module/AnimNextModule_EditorData.h
--- a/Engine/Plugins/Experimental/AnimNext/Source/AnimNextUncookedOnly/Internal/Module/AnimNextModule_EditorData.h
+++ b/Engine/Plugins/Experimental/AnimNext/Source/AnimNextUncookedOnly/Internal/Module/AnimNextModule_EditorData.h
@@ -54,8 +54,8 @@ private:
 	virtual TConstArrayView<TSubclassOf<UAnimNextRigVMAssetEntry>> GetEntryClasses() const override;
 	virtual void RecompileVM() override;
 	virtual void OnPreCompileAsset(FRigVMCompileSettings& InSettings) override;
-	virtual void OnPreCompileGetProgrammaticGraphs(const FRigVMCompileSettings& InSettings, TArray<URigVMGraph*>& OutGraphs) override;
-	virtual void OnPreCompileProcessGraphs(const FRigVMCompileSettings& InSettings, TArray<URigVMGraph*>& InOutGraphs) override;
+	virtual void OnPreCompileGetProgrammaticGraphs(const FRigVMCompileSettings& InSettings, FAnimNextGetGraphCompileContext& OutCompileContext) override;
+	virtual void OnPreCompileProcessGraphs(const FRigVMCompileSettings& InSettings, FAnimNextProcessGraphCompileContext& OutCompileContext) override;
 	virtual void CustomizeNewAssetEntry(UAnimNextRigVMAssetEntry* InNewEntry) const override;
 
 private:
diff --git a/Engine/Plugins/Experimental/AnimNext/Source/AnimNextUncookedOnly/Internal/UncookedOnlyUtils.h b/Engine/Plugins/Experimental/AnimNext/Source/AnimNextUncookedOnly/Internal/UncookedOnlyUtils.h
--- a/Engine/Plugins/Experimental/AnimNext/Source/AnimNextUncookedOnly/Internal/UncookedOnlyUtils.h
+++ b/Engine/Plugins/Experimental/AnimNext/Source/AnimNextUncookedOnly/Internal/UncookedOnlyUtils.h
@@ -15,6 +15,7 @@
 // --- ---
 #include "UncookedOnlyUtils.generated.h"
 
+struct FAnimNextGetVariableCompileContext;
 struct FAnimNextVariableBindingData;
 struct FEdGraphPinType;
 struct FWorkspaceOutlinerItemExports;
@@ -152,7 +153,7 @@ extern ANIMNEXTUNCOOKEDONLY_API TAutoConsoleVariable<bool> CVarDumpProgrammaticG
 
 struct ANIMNEXTUNCOOKEDONLY_API FUtils
 {
-	static void CompileVariables(const FRigVMCompileSettings& InSettings, UAnimNextRigVMAsset* InAsset);
+	static void CompileVariables(const FRigVMCompileSettings& InSettings, UAnimNextRigVMAsset* InAsset, FAnimNextGetVariableCompileContext& OutCompileContext);
 
 	static void CompileVariableBindings(const FRigVMCompileSettings& InSettings, UAnimNextRigVMAsset* InAsset, TArray<URigVMGraph*>& OutGraphs);
 
@@ -239,6 +240,9 @@ struct ANIMNEXTUNCOOKEDONLY_API FUtils
 	static void OpenProgrammaticGraphs(UAnimNextRigVMAssetEditorData* InEditorData, const TArray<URigVMGraph*>& ProgrammaticGraphs);
 #endif // WITH_EDITOR
 
+	// Make a variable name that we use as a wrapper for a function param or return
+	static FString MakeFunctionWrapperVariableName(FName InFunctionName, FName InVariableName);
+
 	// Make an event name that we use as a wrapper to call RigVM functions 
 	static FString MakeFunctionWrapperEventName(FName InFunctionName);
 
diff --git a/Engine/Plugins/Experimental/AnimNext/Source/AnimNextUncookedOnly/Private/AnimNextRigVMAssetEditorData.cpp b/Engine/Plugins/Experimental/AnimNext/Source/AnimNextUncookedOnly/Private/AnimNextRigVMAssetEditorData.cpp
--- a/Engine/Plugins/Experimental/AnimNext/Source/AnimNextUncookedOnly/Private/AnimNextRigVMAssetEditorData.cpp
+++ b/Engine/Plugins/Experimental/AnimNext/Source/AnimNextUncookedOnly/Private/AnimNextRigVMAssetEditorData.cpp
@@ -4,6 +4,11 @@
 #include "AnimNextEdGraph.h"
 #include "AnimNextEdGraphSchema.h"
 #include "AnimNextRigVMAsset.h"
+#include "Compilation/AnimNextRigVMAssetCompileContext.h"
+#include "Compilation/AnimNextGetFunctionHeaderCompileContext.h"
+#include "Compilation/AnimNextGetVariableCompileContext.h"
+#include "Compilation/AnimNextGetGraphCompileContext.h"
+#include "Compilation/AnimNextProcessGraphCompileContext.h"
 #include "Entries/AnimNextRigVMAssetEntry.h"
 #include "AnimNextRigVMAssetSchema.h"
 #include "AnimNextAssetWorkspaceAssetUserData.h"
@@ -790,34 +795,49 @@ void UAnimNextRigVMAssetEditorData::RecompileVM()
 	RigGraphDisplaySettings.MinMicroSeconds = RigGraphDisplaySettings.LastMinMicroSeconds = DBL_MAX;
 	RigGraphDisplaySettings.MaxMicroSeconds = RigGraphDisplaySettings.LastMaxMicroSeconds = (double)INDEX_NONE;
 
-	TArray<URigVMGraph*> ProgrammaticGraphs;
+	FAnimNextRigVMAssetCompileContext CompileContext = {};
 	{
 		TGuardValue<bool> ReentrantGuardSelf(bSuspendModelNotificationsForSelf, true);
 		TGuardValue<bool> ReentrantGuardOthers(RigVMClient.bSuspendModelNotificationsForOthers, true);
 
 		FUtils::RecreateVM(Asset);
 
-		FUtils::CompileVariables(Settings, Asset);
+		{
+			FAnimNextGetFunctionHeaderCompileContext GetFunctionHeaderCompileContext(CompileContext);
+			OnPreCompileGetProgrammaticFunctionHeaders(Settings, GetFunctionHeaderCompileContext);
+		}
+
+		{
+			FAnimNextGetVariableCompileContext GetVariableCompileContext(CompileContext);
+			FUtils::CompileVariables(Settings, Asset, GetVariableCompileContext);
+		}
 
-		OnPreCompileGetProgrammaticGraphs(Settings, ProgrammaticGraphs);
-		for(URigVMGraph* ProgrammaticGraph : ProgrammaticGraphs)
+		{
+			FAnimNextGetGraphCompileContext GetGraphCompileContext(CompileContext);
+			OnPreCompileGetProgrammaticGraphs(Settings, GetGraphCompileContext);
+		}
+
+		for(URigVMGraph* ProgrammaticGraph : CompileContext.ProgrammaticGraphs)
 		{
 			check(ProgrammaticGraph != nullptr);
 		}
 
 		FRigVMClient* VMClient = GetRigVMClient();
 
-		TArray<URigVMGraph*> AllGraphs = VMClient->GetAllModels(false, false);
-		AllGraphs.Append(ProgrammaticGraphs);
+		CompileContext.AllGraphs = VMClient->GetAllModels(false, false);
+		CompileContext.AllGraphs.Append(CompileContext.ProgrammaticGraphs);
 
-		OnPreCompileProcessGraphs(Settings, AllGraphs);
+		{
+			FAnimNextProcessGraphCompileContext ProcessGraphCompileContext(CompileContext);
+			OnPreCompileProcessGraphs(Settings, ProcessGraphCompileContext);
+		}
 
-		if(AllGraphs.Num() > 0)
+		if(CompileContext.AllGraphs.Num() > 0)
 		{
-			URigVMController* Controller = VMClient->GetOrCreateController(AllGraphs[0]);
+			URigVMController* Controller = VMClient->GetOrCreateController(CompileContext.AllGraphs[0]);
 
 			URigVMCompiler* Compiler = URigVMCompiler::StaticClass()->GetDefaultObject<URigVMCompiler>();
-			Compiler->Compile(Settings, AllGraphs, Controller, Asset->VM, Asset->ExtendedExecuteContext, Asset->GetExternalVariables(), &PinToOperandMap);
+			Compiler->Compile(Settings, CompileContext.AllGraphs, Controller, Asset->VM, Asset->ExtendedExecuteContext, Asset->GetExternalVariables(), &PinToOperandMap);
 		}
 
 		// Initialize right away, in packaged builds we initialize during PostLoad
@@ -849,15 +869,15 @@ void UAnimNextRigVMAssetEditorData::RecompileVM()
 		// Display programmatic graphs
 		if(CVarDumpProgrammaticGraphs.GetValueOnGameThread())
 		{
-			FUtils::OpenProgrammaticGraphs(this, ProgrammaticGraphs);
+			FUtils::OpenProgrammaticGraphs(this, CompileContext.ProgrammaticGraphs);
 		}
 		else
 #endif
 		{
-			RemoveProgrammaticGraphs(ProgrammaticGraphs);
+			RemoveProgrammaticGraphs(CompileContext.ProgrammaticGraphs);
 		}
 
-		RemoveTransientGraphs(AllGraphs);
+		RemoveTransientGraphs(CompileContext.AllGraphs);
 
 		OnPostCompileCleanup(Settings);
 
diff --git a/Engine/Plugins/Experimental/AnimNext/Source/AnimNextUncookedOnly/Private/Compilation/AnimNextGetFunctionHeaderCompileContext.cpp b/Engine/Plugins/Experimental/AnimNext/Source/AnimNextUncookedOnly/Private/Compilation/AnimNextGetFunctionHeaderCompileContext.cpp
new file mode 100644
--- /dev/null
+++ b/Engine/Plugins/Experimental/AnimNext/Source/AnimNextUncookedOnly/Private/Compilation/AnimNextGetFunctionHeaderCompileContext.cpp
@@ -0,0 +1,18 @@
+// Copyright Epic Games, Inc. All Rights Reserved.
+
+#include "Compilation/AnimNextGetFunctionHeaderCompileContext.h"
+
+FAnimNextGetFunctionHeaderCompileContext::FAnimNextGetFunctionHeaderCompileContext(FAnimNextRigVMAssetCompileContext& InCompilerContext)
+	: CompilerContext(InCompilerContext)
+{
+}
+
+const TArray<FAnimNextProgrammaticFunctionHeader>& FAnimNextGetFunctionHeaderCompileContext::GetFunctionHeaders() const
+{
+	return CompilerContext.FunctionHeaders;
+}
+
+TArray<FAnimNextProgrammaticFunctionHeader>& FAnimNextGetFunctionHeaderCompileContext::GetMutableFunctionHeaders()
+{
+	return CompilerContext.FunctionHeaders;
+}
diff --git a/Engine/Plugins/Experimental/AnimNext/Source/AnimNextUncookedOnly/Private/Compilation/AnimNextGetGraphCompileContext.cpp b/Engine/Plugins/Experimental/AnimNext/Source/AnimNextUncookedOnly/Private/Compilation/AnimNextGetGraphCompileContext.cpp
new file mode 100644
--- /dev/null
+++ b/Engine/Plugins/Experimental/AnimNext/Source/AnimNextUncookedOnly/Private/Compilation/AnimNextGetGraphCompileContext.cpp
@@ -0,0 +1,23 @@
+// Copyright Epic Games, Inc. All Rights Reserved.
+
+#include "Compilation/AnimNextGetGraphCompileContext.h"
+
+FAnimNextGetGraphCompileContext::FAnimNextGetGraphCompileContext(FAnimNextRigVMAssetCompileContext& InCompilerContext)
+	: CompilerContext(InCompilerContext)
+{
+}
+
+const TArray<FAnimNextProgrammaticFunctionHeader>& FAnimNextGetGraphCompileContext::GetFunctionHeaders() const
+{
+	return CompilerContext.FunctionHeaders;
+}
+
+const TArray<URigVMGraph*>& FAnimNextGetGraphCompileContext::GetProgrammaticGraphs() const
+{
+	return CompilerContext.ProgrammaticGraphs;
+}
+
+TArray<URigVMGraph*>& FAnimNextGetGraphCompileContext::GetMutableProgrammaticGraphs() 
+{
+	return CompilerContext.ProgrammaticGraphs;
+}
diff --git a/Engine/Plugins/Experimental/AnimNext/Source/AnimNextUncookedOnly/Private/Compilation/AnimNextGetVariableCompileContext.cpp b/Engine/Plugins/Experimental/AnimNext/Source/AnimNextUncookedOnly/Private/Compilation/AnimNextGetVariableCompileContext.cpp
new file mode 100644
--- /dev/null
+++ b/Engine/Plugins/Experimental/AnimNext/Source/AnimNextUncookedOnly/Private/Compilation/AnimNextGetVariableCompileContext.cpp
@@ -0,0 +1,23 @@
+// Copyright Epic Games, Inc. All Rights Reserved.
+
+#include "Compilation/AnimNextGetVariableCompileContext.h"
+
+FAnimNextGetVariableCompileContext::FAnimNextGetVariableCompileContext(FAnimNextRigVMAssetCompileContext& InCompilerContext)
+	: CompilerContext(InCompilerContext)
+{
+}
+
+const TArray<FAnimNextProgrammaticFunctionHeader>& FAnimNextGetVariableCompileContext::GetFunctionHeaders() const
+{
+	return CompilerContext.FunctionHeaders;
+}
+
+const TArray<FAnimNextProgrammaticVariable>& FAnimNextGetVariableCompileContext::GetProgrammaticVariables() const
+{
+	return CompilerContext.ProgrammaticVariables;
+}
+
+TArray<FAnimNextProgrammaticVariable>& FAnimNextGetVariableCompileContext::GetMutableProgrammaticVariables()
+{
+	return CompilerContext.ProgrammaticVariables;
+}
diff --git a/Engine/Plugins/Experimental/AnimNext/Source/AnimNextUncookedOnly/Private/Compilation/AnimNextProcessGraphCompileContext.cpp b/Engine/Plugins/Experimental/AnimNext/Source/AnimNextUncookedOnly/Private/Compilation/AnimNextProcessGraphCompileContext.cpp
new file mode 100644
--- /dev/null
+++ b/Engine/Plugins/Experimental/AnimNext/Source/AnimNextUncookedOnly/Private/Compilation/AnimNextProcessGraphCompileContext.cpp
@@ -0,0 +1,18 @@
+// Copyright Epic Games, Inc. All Rights Reserved.
+
+#include "Compilation/AnimNextProcessGraphCompileContext.h"
+
+FAnimNextProcessGraphCompileContext::FAnimNextProcessGraphCompileContext(FAnimNextRigVMAssetCompileContext& InCompilerContext)
+	: CompilerContext(InCompilerContext)
+{
+}
+
+const TArray<URigVMGraph*>& FAnimNextProcessGraphCompileContext::GetAllGraphs() const
+{
+	return CompilerContext.AllGraphs;
+}
+
+TArray<URigVMGraph*>& FAnimNextProcessGraphCompileContext::GetMutableAllGraphs() 
+{
+	return CompilerContext.AllGraphs;
+}
diff --git a/Engine/Plugins/Experimental/AnimNext/Source/AnimNextUncookedOnly/Private/Module/AnimNextModule_EditorData.cpp b/Engine/Plugins/Experimental/AnimNext/Source/AnimNextUncookedOnly/Private/Module/AnimNextModule_EditorData.cpp
--- a/Engine/Plugins/Experimental/AnimNext/Source/AnimNextUncookedOnly/Private/Module/AnimNextModule_EditorData.cpp
+++ b/Engine/Plugins/Experimental/AnimNext/Source/AnimNextUncookedOnly/Private/Module/AnimNextModule_EditorData.cpp
@@ -2,6 +2,8 @@
 
 #include "Module/AnimNextModule_EditorData.h"
 
+#include "Compilation/AnimNextGetGraphCompileContext.h"
+#include "Compilation/AnimNextProcessGraphCompileContext.h"
 #include "ExternalPackageHelper.h"
 #include "UncookedOnlyUtils.h"
 #include "Module/AnimNextModule.h"
@@ -51,21 +53,21 @@ void UAnimNextModule_EditorData::OnPreCompileAsset(FRigVMCompileSettings& InSett
 	Module->RequiredComponents.Empty();
 }
 
-void UAnimNextModule_EditorData::OnPreCompileGetProgrammaticGraphs(const FRigVMCompileSettings& InSettings, TArray<URigVMGraph*>& OutGraphs)
+void UAnimNextModule_EditorData::OnPreCompileGetProgrammaticGraphs(const FRigVMCompileSettings& InSettings, FAnimNextGetGraphCompileContext& OutCompileContext)
 {
 	using namespace UE::AnimNext::UncookedOnly;
 
-	FUtils::CompileVariableBindings(InSettings, FUtils::GetAsset<UAnimNextModule>(this), OutGraphs);
+	FUtils::CompileVariableBindings(InSettings, FUtils::GetAsset<UAnimNextModule>(this), OutCompileContext.GetMutableProgrammaticGraphs());
 }
 
-void UAnimNextModule_EditorData::OnPreCompileProcessGraphs(const FRigVMCompileSettings& InSettings, TArray<URigVMGraph*>& InOutGraphs)
+void UAnimNextModule_EditorData::OnPreCompileProcessGraphs(const FRigVMCompileSettings& InSettings, FAnimNextProcessGraphCompileContext& OutCompileContext)
 {
 	using namespace UE::AnimNext::UncookedOnly;
 
 	UAnimNextModule* Module = FUtils::GetAsset<UAnimNextModule>(this);
 
 	// Gather any required components from node metadata
-	for(URigVMGraph* Graph : InOutGraphs)
+	for(URigVMGraph* Graph : OutCompileContext.GetMutableAllGraphs())
 	{
 		for(URigVMNode* Node : Graph->GetNodes())
 		{
diff --git a/Engine/Plugins/Experimental/AnimNext/Source/AnimNextUncookedOnly/Private/UncookedOnlyUtils.cpp b/Engine/Plugins/Experimental/AnimNext/Source/AnimNextUncookedOnly/Private/UncookedOnlyUtils.cpp
--- a/Engine/Plugins/Experimental/AnimNext/Source/AnimNextUncookedOnly/Private/UncookedOnlyUtils.cpp
+++ b/Engine/Plugins/Experimental/AnimNext/Source/AnimNextUncookedOnly/Private/UncookedOnlyUtils.cpp
@@ -5,6 +5,7 @@
 #include "K2Node_CallFunction.h"
 #include "AssetRegistry/AssetRegistryModule.h"
 #include "AssetRegistry/IAssetRegistry.h"
+#include "Compilation/AnimNextGetVariableCompileContext.h"
 #include "Module/AnimNextModule.h"
 #include "Module/AnimNextModule_EditorData.h"
 #include "Serialization/MemoryReader.h"
@@ -52,7 +53,7 @@ void FUtils::RecreateVM(UAnimNextRigVMAsset* InAsset)
 }
 
 
-void FUtils::CompileVariables(const FRigVMCompileSettings& InSettings, UAnimNextRigVMAsset* InAsset)
+void FUtils::CompileVariables(const FRigVMCompileSettings& InSettings, UAnimNextRigVMAsset* InAsset, FAnimNextGetVariableCompileContext& OutCompileContext)
 {
 	check(InAsset);
 
@@ -68,8 +69,8 @@ void FUtils::CompileVariables(const FRigVMCompileSettings& InSettings, UAnimNext
 	UAnimNextDataInterface_EditorData* EditorData = GetEditorData<UAnimNextDataInterface_EditorData>(DataInterface);
 
 	// Gather programmatic variables regenerated each compile
-	TArray<FAnimNextProgrammaticVariable> ProgrammaticVariables;
-	EditorData->OnPreCompileGetProgrammaticVariables(InSettings, ProgrammaticVariables);
+	EditorData->OnPreCompileGetProgrammaticVariables(InSettings, OutCompileContext);
+	const TArray<FAnimNextProgrammaticVariable>& ProgrammaticVariables = OutCompileContext.GetProgrammaticVariables();
 
 	struct FStructEntryInfo
 	{
@@ -1169,6 +1170,12 @@ void FUtils::OpenProgrammaticGraphs(UAnimNextRigVMAssetEditorData* EditorData, c
 }
 #endif // WITH_EDITOR
 
+FString FUtils::MakeFunctionWrapperVariableName(FName InFunctionName, FName InVariableName)
+{
+	// We assume the function name is enough for variable name uniqueness in this graph (We don't yet desire global uniqueness).
+	return TEXT("__InternalVar_") + InFunctionName.ToString() + "_" + InVariableName.ToString();
+}
+
 FString FUtils::MakeFunctionWrapperEventName(FName InFunctionName)
 {
 	return TEXT("__InternalCall_") + InFunctionName.ToString();
diff --git a/Engine/Plugins/Experimental/AnimNextAnimGraph/Source/AnimNextAnimGraphUncookedOnly/Internal/Graph/AnimNextAnimationGraph_EditorData.h b/Engine/Plugins/Experimental/AnimNextAnimGraph/Source/AnimNextAnimGraphUncookedOnly/Internal/Graph/AnimNextAnimationGraph_EditorData.h
--- a/Engine/Plugins/Experimental/AnimNextAnimGraph/Source/AnimNextAnimGraphUncookedOnly/Internal/Graph/AnimNextAnimationGraph_EditorData.h
+++ b/Engine/Plugins/Experimental/AnimNextAnimGraph/Source/AnimNextAnimGraphUncookedOnly/Internal/Graph/AnimNextAnimationGraph_EditorData.h
@@ -72,8 +72,10 @@ protected:
 	virtual TSubclassOf<UAssetUserData> GetAssetUserDataClass() const override;
 	virtual void InitializeAssetUserData() override;
 	virtual void OnPreCompileAsset(FRigVMCompileSettings& InSettings) override;
-	virtual void OnPreCompileGetProgrammaticGraphs(const FRigVMCompileSettings& InSettings, TArray<URigVMGraph*>& OutGraphs) override;
-	virtual void OnPreCompileProcessGraphs(const FRigVMCompileSettings& InSettings, TArray<URigVMGraph*>& InOutGraphs) override;
+	virtual void OnPreCompileGetProgrammaticFunctionHeaders(const FRigVMCompileSettings& InSettings, FAnimNextGetFunctionHeaderCompileContext& OutCompileContext) override;
+	virtual void OnPreCompileGetProgrammaticVariables(const FRigVMCompileSettings& InSettings, FAnimNextGetVariableCompileContext& OutCompileContext) override;
+	virtual void OnPreCompileGetProgrammaticGraphs(const FRigVMCompileSettings& InSettings, FAnimNextGetGraphCompileContext& OutCompileContext) override;
+	virtual void OnPreCompileProcessGraphs(const FRigVMCompileSettings& InSettings, FAnimNextProcessGraphCompileContext& OutCompileContext) override;
 	virtual void OnPostCompileCleanup(const FRigVMCompileSettings& InSettings) override;
 
 	// Super overrides
diff --git a/Engine/Plugins/Experimental/AnimNextAnimGraph/Source/AnimNextAnimGraphUncookedOnly/Private/Graph/AnimNextAnimationGraph_EditorData.cpp b/Engine/Plugins/Experimental/AnimNextAnimGraph/Source/AnimNextAnimGraphUncookedOnly/Private/Graph/AnimNextAnimationGraph_EditorData.cpp
--- a/Engine/Plugins/Experimental/AnimNextAnimGraph/Source/AnimNextAnimGraphUncookedOnly/Private/Graph/AnimNextAnimationGraph_EditorData.cpp
+++ b/Engine/Plugins/Experimental/AnimNextAnimGraph/Source/AnimNextAnimGraphUncookedOnly/Private/Graph/AnimNextAnimationGraph_EditorData.cpp
@@ -4,6 +4,10 @@
 #include "AnimGraphUncookedOnlyUtils.h"
 #include "AnimNextAnimGraphWorkspaceAssetUserData.h"
 #include "AnimNextScopedCompilerResults.h"
+#include "Compilation/AnimNextGetFunctionHeaderCompileContext.h"
+#include "Compilation/AnimNextGetVariableCompileContext.h"
+#include "Compilation/AnimNextGetGraphCompileContext.h"
+#include "Compilation/AnimNextProcessGraphCompileContext.h"
 #include "UncookedOnlyUtils.h"
 #include "AssetRegistry/IAssetRegistry.h"
 #include "Curves/CurveFloat.h"
@@ -30,6 +34,7 @@
 #include "Entries/AnimNextRigVMAssetEntry.h"
 #include "Graph/RigUnit_AnimNextGraphEvaluator.h"
 #include "Graph/RigUnit_AnimNextGraphRoot.h"
+#include "Variables/AnimNextProgrammaticVariable.h"
 
 #define LOCTEXT_NAMESPACE "AnimNextAnimationGraph_EditorData"
 
@@ -491,13 +496,103 @@
 	AnimationGraph->DefaultEntryPoint = NAME_None;
 }
 
-void UAnimNextAnimationGraph_EditorData::OnPreCompileProcessGraphs(const FRigVMCompileSettings& InSettings, TArray<URigVMGraph*>& InOutGraphs)
+void UAnimNextAnimationGraph_EditorData::OnPreCompileGetProgrammaticFunctionHeaders(const FRigVMCompileSettings& InSettings, FAnimNextGetFunctionHeaderCompileContext& OutCompileContext)
+{
+	using namespace UE::AnimNext::UncookedOnly;
+	using namespace UE::AnimNext::UncookedOnly::Private;
+
+	Super::OnPreCompileGetProgrammaticFunctionHeaders(InSettings, OutCompileContext);
+
+	// Gather all 'call function' traits and create shim-calls for them.
+	// For the compiler to pick them up if they are not public we need a calling reference to the function from a graph
+	FRigVMClient* VMClient = GetRigVMClient();
+	for (const URigVMGraph* Graph : VMClient->GetAllModels(false, false))
+	{
+		for(const URigVMNode* Node : Graph->GetNodes())
+		{
+			for(const URigVMPin* TraitPin : Node->GetTraitPins())
+			{
+				if (TraitPin->IsExecuteContext())
+				{
+					continue;
+				}
+
+				TSharedPtr<FStructOnScope> ScopedTrait = Node->GetTraitInstance(TraitPin->GetFName());
+				if (!ScopedTrait.IsValid())
+				{
+					continue;
+				}
+
+				const FRigVMTrait* Trait = (FRigVMTrait*)ScopedTrait->GetStructMemory();
+				UScriptStruct* TraitSharedInstanceData = Trait->GetTraitSharedDataStruct();
+				if (TraitSharedInstanceData == nullptr)
+				{
+					continue;
+				}
+
+				if(!TraitSharedInstanceData->IsChildOf(FAnimNextCallFunctionSharedData::StaticStruct()))
+				{
+					continue;
+				}
+
+				const FString DefaultValue = TraitPin->GetDefaultValue();
+				TInstancedStruct<FAnimNextCallFunctionSharedData> InstancedStruct = TInstancedStruct<FAnimNextCallFunctionSharedData>::Make();
+				FRigVMPinDefaultValueImportErrorContext ErrorPipe(ELogVerbosity::Verbose);
+				LOG_SCOPE_VERBOSITY_OVERRIDE(LogExec, ErrorPipe.GetMaxVerbosity());
+				TraitSharedInstanceData->ImportText(*DefaultValue, InstancedStruct.GetMutableMemory(), nullptr, PPF_SerializedAsImportText, &ErrorPipe, TraitSharedInstanceData->GetName());
+
+				const FRigVMGraphFunctionHeader& FunctionHeader = InstancedStruct.Get<FAnimNextCallFunctionSharedData>().FunctionHeader;
+				if(FunctionHeader.IsValid())
+				{
+					FAnimNextProgrammaticFunctionHeader AnimNextFunctionHeader = {};
+					AnimNextFunctionHeader.Wrapped = FunctionHeader;
+					// @TODO: Determine if param / return variables are needed based on 'FAnimNextCallFunctionSharedData'
+					//AnimNextFunctionHeader.bGenerateParamVariables = true;
+					//AnimNextFunctionHeader.bGenerateReturnVariables = true;
+					OutCompileContext.GetMutableFunctionHeaders().Add(AnimNextFunctionHeader);
+				}
+			}
+		}
+	}
+}
+
+void UAnimNextAnimationGraph_EditorData::OnPreCompileGetProgrammaticVariables(const FRigVMCompileSettings& InSettings, FAnimNextGetVariableCompileContext& OutCompileContext)
+{
+	using namespace UE::AnimNext::UncookedOnly;
+
+	Super::OnPreCompileGetProgrammaticVariables(InSettings, OutCompileContext);
+
+	for (const FAnimNextProgrammaticFunctionHeader& ProgrammaticFunctionHeader : OutCompileContext.GetFunctionHeaders())
+	{
+		if (!ProgrammaticFunctionHeader.bGenerateParamVariables && !ProgrammaticFunctionHeader.bGenerateReturnVariables)
+		{
+			continue;
+		}
+
+		const FRigVMGraphFunctionHeader& FunctionHeader = ProgrammaticFunctionHeader.Wrapped;
+		for (const FRigVMGraphFunctionArgument& Argument : FunctionHeader.Arguments)
+		{
+			bool bAddParam = ProgrammaticFunctionHeader.bGenerateParamVariables && Argument.Direction == ERigVMPinDirection::Input;
+			bool bAddReturn = ProgrammaticFunctionHeader.bGenerateReturnVariables && Argument.Direction == ERigVMPinDirection::Output;
+
+			if (bAddParam || bAddReturn)
+			{
+				FRigVMGraphFunctionArgument InternallyNamedArgument = Argument;
+				InternallyNamedArgument.Name = FName(FUtils::MakeFunctionWrapperVariableName(FunctionHeader.Name, Argument.Name));
+				OutCompileContext.GetMutableProgrammaticVariables().Add(FAnimNextProgrammaticVariable::FromRigVMGraphFunctionArgument(InternallyNamedArgument));
+			}
+		}
+	}
+}
+
+void UAnimNextAnimationGraph_EditorData::OnPreCompileProcessGraphs(const FRigVMCompileSettings& InSettings, FAnimNextProcessGraphCompileContext& OutCompileContext)
 {
 	using namespace UE::AnimNext::UncookedOnly;
 	using namespace UE::AnimNext::UncookedOnly::Private;
 
 	FRigVMClient* VMClient = GetRigVMClient();
 	UAnimNextAnimationGraph* AnimationGraph = FUtils::GetAsset<UAnimNextAnimationGraph>(this);
+	TArray<URigVMGraph*>& InOutGraphs = OutCompileContext.GetMutableAllGraphs();
 
 	TArray<URigVMGraph*> AnimGraphs;
 	TArray<URigVMGraph*> NonAnimGraphs;
@@ -772,72 +867,27 @@
 	Super::InitializeAssetUserData();
 }
 
-void UAnimNextAnimationGraph_EditorData::OnPreCompileGetProgrammaticGraphs(const FRigVMCompileSettings& InSettings, TArray<URigVMGraph*>& OutGraphs)
+void UAnimNextAnimationGraph_EditorData::OnPreCompileGetProgrammaticGraphs(const FRigVMCompileSettings& InSettings, FAnimNextGetGraphCompileContext& OutCompileContext)
 {
 	using namespace UE::AnimNext::UncookedOnly;
 
-	Super::OnPreCompileGetProgrammaticGraphs(InSettings, OutGraphs);
+	Super::OnPreCompileGetProgrammaticGraphs(InSettings, OutCompileContext);
 
-	// Gather all 'call function' traits and create shim-calls for them.
-	// For the compiler to pick them up if they are not public we need a calling reference to the function from a graph
-	TArray<FRigVMGraphFunctionHeader> FunctionHeaders;
+	if(OutCompileContext.GetFunctionHeaders().Num() > 0)
+	{
+		FRigVMClient* VMClient = GetRigVMClient();
 
-	UAnimNextAnimationGraph* Asset = FUtils::GetAsset<UAnimNextAnimationGraph>(this);
-	FRigVMClient* VMClient = GetRigVMClient();
-	for (const URigVMGraph* Graph : VMClient->GetAllModels(false, false))
-	{
-		for(const URigVMNode* Node : Graph->GetNodes())
-		{
-			for(const URigVMPin* TraitPin : Node->GetTraitPins())
-			{
-				if (TraitPin->IsExecuteContext())
-				{
-					continue;
-				}
-
-				TSharedPtr<FStructOnScope> ScopedTrait = Node->GetTraitInstance(TraitPin->GetFName());
-				if (!ScopedTrait.IsValid())
-				{
-					continue;
-				}
-
-				const FRigVMTrait* Trait = (FRigVMTrait*)ScopedTrait->GetStructMemory();
-				UScriptStruct* TraitSharedInstanceData = Trait->GetTraitSharedDataStruct();
-				if (TraitSharedInstanceData == nullptr)
-				{
-					continue;
-				}
-
-				if(!TraitSharedInstanceData->IsChildOf(FAnimNextCallFunctionSharedData::StaticStruct()))
-				{
-					continue;
-				}
-
-				const FString DefaultValue = TraitPin->GetDefaultValue();
-				TInstancedStruct<FAnimNextCallFunctionSharedData> InstancedStruct = TInstancedStruct<FAnimNextCallFunctionSharedData>::Make();
-				FRigVMPinDefaultValueImportErrorContext ErrorPipe(ELogVerbosity::Verbose);
-				LOG_SCOPE_VERBOSITY_OVERRIDE(LogExec, ErrorPipe.GetMaxVerbosity());
-				TraitSharedInstanceData->ImportText(*DefaultValue, InstancedStruct.GetMutableMemory(), nullptr, PPF_SerializedAsImportText, &ErrorPipe, TraitSharedInstanceData->GetName());
-
-				const FRigVMGraphFunctionHeader& FunctionHeader = InstancedStruct.Get<FAnimNextCallFunctionSharedData>().FunctionHeader;
-				if(FunctionHeader.IsValid())
-				{
-					FunctionHeaders.Add(FunctionHeader);
-				}
-			}
-		}
-	}
-
-	if(FunctionHeaders.Num() > 0)
-	{
 		// Create all shim events for our traits to call
 		const bool bSetupUndoRedo = false;
 		URigVMGraph* WrapperGraph = NewObject<URigVMGraph>(this, NAME_None, RF_Transient);
 		UAnimNextController* Controller = CastChecked<UAnimNextController>(VMClient->GetOrCreateController(WrapperGraph));
 		FRigVMControllerNotifGuard NotifGuard(Controller);
 		bool bAddedWrapperEvent = true;
-		for(const FRigVMGraphFunctionHeader& FunctionHeader : FunctionHeaders)
+
+		for(const FAnimNextProgrammaticFunctionHeader& AnimNextFunctionHeader : OutCompileContext.GetFunctionHeaders())
 		{
+			const FRigVMGraphFunctionHeader& FunctionHeader = AnimNextFunctionHeader.Wrapped;
+
 			URigVMLibraryNode* LibraryNode = Cast<URigVMLibraryNode>(FunctionHeader.LibraryPointer.GetNodeSoftPath().TryLoad());
 			if(LibraryNode == nullptr)
 			{
@@ -869,21 +919,105 @@
 				continue;
 			}
 
-			// Link up nodes
-			URigVMPin* ExecuteOutputPin = EventNode->FindPin(FRigVMStruct::ExecuteContextName.ToString());
+			// Link up Execute nodes if needed, function may be pure & lack an input pin
+			URigVMPin* CurrentExecuteOutputPin = EventNode->FindPin(FRigVMStruct::ExecuteContextName.ToString());
 			URigVMPin* ExecuteInputPin = FunctionNode->FindPin(FRigVMStruct::ExecuteContextName.ToString());
-			if(!Controller->AddLink(ExecuteOutputPin, ExecuteInputPin, bSetupUndoRedo))
+			if (ExecuteInputPin && !Controller->AddLink(CurrentExecuteOutputPin, ExecuteInputPin, bSetupUndoRedo))
 			{
 				InSettings.ReportError(FString::Printf(TEXT("Could not link execute pins for function '%s'"), *FunctionHeader.Name.ToString()));
 				continue;
 			}
 
+			// Update current execute pin, RigVM doesn't have a concept of input / output execute pins, just one execute content pin used for both
+			CurrentExecuteOutputPin = ExecuteInputPin ? ExecuteInputPin : CurrentExecuteOutputPin;
+
+			// Generate & link internal variables if desired
+			if (AnimNextFunctionHeader.bGenerateParamVariables || AnimNextFunctionHeader.bGenerateReturnVariables)
+			{
+				// Controller needs to notify the AST of variable changes to make new links
+				bool bSuspendNotificationForInternalVariables = false;
+				FRigVMControllerNotifGuard VarNotifGuard(Controller, bSuspendNotificationForInternalVariables);
+
+				// Generate & link input arguments, also generate result variable node but link later
+				for (const FRigVMGraphFunctionArgument& Argument : FunctionHeader.Arguments)
+				{
+					bool bIsGetter = Argument.Direction == ERigVMPinDirection::Input;
+
+					if (bIsGetter && AnimNextFunctionHeader.bGenerateParamVariables)
+					{
+						FName InternalVariableName = FName(FUtils::MakeFunctionWrapperVariableName(FunctionHeader.Name, Argument.Name));
+						URigVMVariableNode* FunctionParamVariableNode = Controller->AddVariableNode(InternalVariableName
+							, Argument.CPPType.ToString()
+							, Argument.CPPTypeObject.Get()
+							, bIsGetter
+							, Argument.DefaultValue
+							, FVector2D::ZeroVector
+							, InternalVariableName.ToString()
+							, bSetupUndoRedo);
+
+						if (!FunctionParamVariableNode)
+						{
+							InSettings.ReportError(FString::Printf(TEXT("Failed to add internal variable node for param: %s, var: %s"), *FunctionHeader.Name.ToString(), *InternalVariableName.ToString()));
+							return;
+						}
+
+						// Link Param Pins
+						URigVMPin* ParamValuePin = FunctionParamVariableNode->GetValuePin();
+						URigVMPin* FunctionArgumentPin = FunctionNode->FindPin(Argument.Name.ToString());
+						if (!Controller->AddLink(ParamValuePin, FunctionArgumentPin, bSetupUndoRedo))
+						{
+							InSettings.ReportError(FString::Printf(TEXT("Failed to link internal variable param node to function: %s -> %s"), *GetNameSafe(ParamValuePin), *GetNameSafe(FunctionArgumentPin)));
+							return;
+						}
+					}
+
+					if (!bIsGetter && AnimNextFunctionHeader.bGenerateReturnVariables)
+					{
+						FName InternalResultName = FName(FUtils::MakeFunctionWrapperVariableName(FunctionHeader.Name, Argument.Name));
+						URigVMVariableNode* FunctionResultVariableNode = Controller->AddVariableNode(InternalResultName
+							, Argument.CPPType.ToString()
+							, Argument.CPPTypeObject.Get()
+							, bIsGetter
+							, Argument.DefaultValue
+							, FVector2D::ZeroVector
+							, InternalResultName.ToString()
+							, bSetupUndoRedo);
+
+						if (!FunctionResultVariableNode)
+						{
+							InSettings.ReportError(FString::Printf(TEXT("Failed to add internal variable node for result: %s, var: %s"), *FunctionHeader.Name.ToString(), *InternalResultName.ToString()));
+							return;
+						}
+
+						// Link Result pins
+						URigVMPin* FunctionResultPin = FunctionNode->FindPin(Argument.Name.ToString());
+						URigVMPin* ResultValuePin = FunctionResultVariableNode->GetValuePin();
+						if (!Controller->AddLink(FunctionResultPin, ResultValuePin, bSetupUndoRedo))
+						{
+							InSettings.ReportError(FString::Printf(TEXT("Failed to link internal variable result node to function: %s -> %s"), *GetNameSafe(FunctionResultPin), *GetNameSafe(ResultValuePin)));
+							return;
+						}
+
+						// Link Result Execute pins
+						URigVMPin* ResultExecuteInputPin = FunctionResultVariableNode->FindPin(FRigVMStruct::ExecuteContextName.ToString());
+						if (!Controller->AddLink(CurrentExecuteOutputPin, ResultExecuteInputPin, bSetupUndoRedo))
+						{
+							InSettings.ReportError(FString::Printf(TEXT("Failed to link execute pins for variable result node: %s -> %s"), *GetNameSafe(CurrentExecuteOutputPin), *GetNameSafe(ResultExecuteInputPin)));
+							return;
+						}
+
+						// Update current execute pin, RigVM doesn't have a concept of input / output execute pins, just one execute content pin used for both
+						CurrentExecuteOutputPin = ResultExecuteInputPin;
+					}
+				}
+			}
+
 			bAddedWrapperEvent = true;
 		}
 
 		if(bAddedWrapperEvent)
 		{
-			OutGraphs.Add(WrapperGraph);
+			OutCompileContext.GetMutableProgrammaticGraphs().Add(WrapperGraph);
 		}
 	}
 }
diff --git a/Engine/Plugins/Experimental/AnimNextStateTree/Source/AnimNextStateTree/Internal/AnimNextStateTreeEditorOnlyTypes.h b/Engine/Plugins/Experimental/AnimNextStateTree/Source/AnimNextStateTree/Internal/AnimNextStateTreeEditorOnlyTypes.h
--- a/Engine/Plugins/Experimental/AnimNextStateTree/Source/AnimNextStateTree/Internal/AnimNextStateTreeEditorOnlyTypes.h
+++ b/Engine/Plugins/Experimental/AnimNextStateTree/Source/AnimNextStateTree/Internal/AnimNextStateTreeEditorOnlyTypes.h
@@ -5,7 +5,7 @@
 #if WITH_EDITOR
 
 #include "CoreMinimal.h"
-#include "Variables/AnimNextProgrammaticVariable.h"
+#include "Compilation/AnimNextGetFunctionHeaderCompileContext.h"
 
 class UAnimNextDataInterface_EditorData;
 class URigVMGraph;
@@ -16,39 +16,20 @@ struct FRigVMClient;
 /**
  * Helper struct to contain info needed to compile RigVM Variables
  */
-struct ANIMNEXTSTATETREE_API FAnimNextStateTreeProgrammaticVariableParams
+struct ANIMNEXTSTATETREE_API FAnimNextStateTreeProgrammaticFunctionHeaderParams
 {
-	FAnimNextStateTreeProgrammaticVariableParams(UAnimNextDataInterface_EditorData* InEditorData, const FRigVMCompileSettings& InSettings, FRigVMClient& InRigVMClient, TArray<FAnimNextProgrammaticVariable>& InOutVariables)
+	FAnimNextStateTreeProgrammaticFunctionHeaderParams(UAnimNextDataInterface_EditorData* InEditorData, const FRigVMCompileSettings& InSettings, FRigVMClient& InRigVMClient, FAnimNextGetFunctionHeaderCompileContext& InOutCompileContext)
 		: EditorData(InEditorData)
 		, Settings(InSettings)
 		, RigVMClient(InRigVMClient)
-		, OutVariables(InOutVariables)
+		, OutCompileContext(InOutCompileContext)
 	{
 	}
 
 	UAnimNextDataInterface_EditorData* EditorData;
 	const FRigVMCompileSettings& Settings;
 	FRigVMClient& RigVMClient;
-	TArray<FAnimNextProgrammaticVariable>& OutVariables;
-};
-
-/**
- * Helper struct to contain info needed to compile RigVM Graphs
- */
-struct ANIMNEXTSTATETREE_API FAnimNextStateTreeProgrammaticGraphParams
-{
-	FAnimNextStateTreeProgrammaticGraphParams(UAnimNextDataInterface_EditorData* InEditorData, const FRigVMCompileSettings& InSettings, FRigVMClient& InRigVMClient, TArray<URigVMGraph*>& InOutGraphs)
-		: EditorData(InEditorData)
-		, Settings(InSettings)
-		, RigVMClient(InRigVMClient)
-		, OutGraphs(InOutGraphs)
-	{
-	}
-
-	UAnimNextDataInterface_EditorData* EditorData;
-	const FRigVMCompileSettings& Settings;
-	FRigVMClient& RigVMClient;
-	TArray<URigVMGraph*>& OutGraphs;
+	FAnimNextGetFunctionHeaderCompileContext& OutCompileContext;
 };
 
 #endif // WITH_EDITOR
\ No newline at end of file
diff --git a/Engine/Plugins/Experimental/AnimNextStateTree/Source/AnimNextStateTree/Private/Conditions/AnimNextStateTreeRigVMConditionBase.cpp b/Engine/Plugins/Experimental/AnimNextStateTree/Source/AnimNextStateTree/Private/Conditions/AnimNextStateTreeRigVMConditionBase.cpp
--- a/Engine/Plugins/Experimental/AnimNextStateTree/Source/AnimNextStateTree/Private/Conditions/AnimNextStateTreeRigVMConditionBase.cpp
+++ b/Engine/Plugins/Experimental/AnimNextStateTree/Source/AnimNextStateTree/Private/Conditions/AnimNextStateTreeRigVMConditionBase.cpp
@@ -10,29 +10,9 @@
 #include "TraitCore/ExecutionContext.h"
 
 #if WITH_EDITOR
-#include "AnimGraphUncookedOnlyUtils.h"
-#include "AnimNextRigVMAssetEditorData.h"
 #include "AnimNextStateTreeEditorOnlyTypes.h"
-#include "DataInterface/AnimNextDataInterfaceFactory.h"
-#include "DataInterface/AnimNextDataInterface.h"
-#include "DataInterface/AnimNextDataInterface_EditorData.h"
-#include "Factories/Factory.h"
-#include "Graph/AnimNextAnimationGraphSchema.h"
-#include "Graph/RigUnit_AnimNextGraphRoot.h"
-#include "Graph/RigUnit_AnimNextShimRoot.h"
-#include "Param/ParamType.h"
 #include "RigVMBlueprintGeneratedClass.h"
-#include "RigVMCompiler/RigVMCompiler.h"
-#include "RigVMFunctions/Execution/RigVMFunction_UserDefinedEvent.h"
 #include "RigVMFunctions/RigVMFunctionDefines.h"
-#include "RigVMModel/RigVMClient.h"
-#include "RigVMModel/Nodes/RigVMFunctionReferenceNode.h"
-#include "RigVMModel/Nodes/RigVMLibraryNode.h"
-#include "RigVMModel/Nodes/RigVMUnitNode.h"
-#include "RigVMModel/Nodes/RigVMVariableNode.h"
-#include "RigVMModel/RigVMFunctionLibrary.h"
-#include "RigVMModel/RigVMGraph.h"
-#include "RigVMTypeUtils.h"
 #include "StateTree.h"
 #include "StateTreeEditorData.h"
 #include "UncookedOnlyUtils.h"
@@ -42,114 +22,9 @@
 
 #include UE_INLINE_GENERATED_CPP_BY_NAME(AnimNextStateTreeRigVMConditionBase)
 
-////////////////////////////////////////////////////////////////////////////////////////////////////////////////////////
-// FAnimNextStateTreeRigVMConditionInstanceData
-
-bool FAnimNextStateTreeRigVMConditionInstanceData::IsIndexCacheInitialized() const
-{
-	return VariableIndexCache.bInitialized;
-}
-
-bool FAnimNextStateTreeRigVMConditionInstanceData::TryPopulateIndexCache(const FAnimNextStateTreeRigVMConditionBase& OwningCondition, TObjectPtr<UAnimNextDataInterface> InDataInterface, FName InResultName, FName InFunctionName)
-{
-	if (VariableIndexCache.bInitialized)
-	{
-		return true;
-	}
-
-	VariableIndexCache.ArgumentIndexes.Reserve(ParamData.Num());
-	TArray<FRigVMExternalVariable> ExternalVariables = InDataInterface->GetExternalVariables();
-	for (const FPropertyBagPropertyDesc& Desc : ParamData.GetPropertyBagStruct()->GetPropertyDescs())
-	{
-		for (int32 Index = 0; Index < ExternalVariables.Num(); Index++)
-		{
-			FName VariableName = ExternalVariables[Index].Name;
-			if (VariableName == OwningCondition.MakeInternalVariableName(Desc.Name))
-			{
-				VariableIndexCache.ArgumentIndexes.Add(Index);
-				break;
-			}
-		}
-	}
-
-	// Search all variables for result, it won't be part of the arg param data
-	bool bResultFound = false;
-	for (int32 Index = 0; Index < ExternalVariables.Num(); Index++)
-	{
-		FName VariableName = ExternalVariables[Index].Name;
-		if (VariableName == InResultName)
-		{
-			VariableIndexCache.ResultIndex = Index;
-			bResultFound = true;
-			break;
-		}
-	}
-
-	bool bAllArgsFound = ensure(VariableIndexCache.ArgumentIndexes.Num() == ParamData.Num());
-	if (bAllArgsFound && ensure(bResultFound))
-	{
-		VariableIndexCache.bInitialized = true;
-	}
-
-	return VariableIndexCache.bInitialized;
-}
-
-TArrayView<const uint8> FAnimNextStateTreeRigVMConditionInstanceData::GetVMArgumentIndexes() const
-{
-	return MakeArrayView(VariableIndexCache.ArgumentIndexes);
-}
-
-uint8 FAnimNextStateTreeRigVMConditionInstanceData::GetVMResultIndex() const
-{
-	return VariableIndexCache.ResultIndex;
-}
-
 ////////////////////////////////////////////////////////////////////////////////////////////////////////////////////////
 // FAnimNextStateTreeRigVMConditionBase
 
-const FName FAnimNextStateTreeRigVMConditionBase::InternalEventPrefixName("__Internal__Event__");
-const FName FAnimNextStateTreeRigVMConditionBase::InternalVariablePrefixName("__Internal__Variable__");
-
-FName FAnimNextStateTreeRigVMConditionBase::MakeInternalEventName(const FName InName) const
-{
-	ensure(!StateName.IsNone());
-	ensure(!RigVMFunctionHeader.Name.IsNone());
-	ensure(NodeId.IsValid());
-
-	// We need to append the variable name to avoid conflicts with other same named variables used by different functions in the VM (Ex: Result)
-	TStringBuilder<512> StringBuilder;
-	StringBuilder << InternalEventPrefixName.ToString();
-	StringBuilder << StateName.ToString();
-	StringBuilder << TEXT("__");
-	StringBuilder << RigVMFunctionHeader.Name.ToString();
-	StringBuilder << TEXT("__");
-	StringBuilder << InName.ToString();
-	StringBuilder << TEXT("__");
-	StringBuilder << NodeId.ToString();
-
-	return FName(StringBuilder.ToString());
-}
-
-FName FAnimNextStateTreeRigVMConditionBase::MakeInternalVariableName(const FName InName) const
-{
-	ensure(!StateName.IsNone());
-	ensure(!RigVMFunctionHeader.Name.IsNone());
-	ensure(NodeId.IsValid());
-
-	// We need to append the variable name to avoid conflicts with other same named variables used by different functions in the VM (Ex: Result)
-	TStringBuilder<512> StringBuilder;
-	StringBuilder << InternalVariablePrefixName.ToString();
-	StringBuilder << StateName.ToString();
-	StringBuilder << TEXT("__");
-	StringBuilder << RigVMFunctionHeader.Name.ToString();
-	StringBuilder << TEXT("__");
-	StringBuilder << InName.ToString();
-	StringBuilder << TEXT("__");
-	StringBuilder << NodeId.ToString();
-
-	return FName(StringBuilder.ToString());
-}
-
 bool FAnimNextStateTreeRigVMConditionBase::Link(FStateTreeLinker& Linker)
 {
 	Linker.LinkExternalData(TraitContextHandle);
@@ -167,14 +42,14 @@ bool FAnimNextStateTreeRigVMConditionBase::TestCondition(FStateTreeExecutionCont
 		if (TObjectPtr<UAnimNextDataInterface> RigVMInstance = const_cast<UAnimNextDataInterface*>(AnimExecContext->GetRootGraphInstance().GetDataInterface()))
 		{
 			// Cache RigVM variable indexes, which can vary as they are parent derived
-			if (!InstanceData->IsIndexCacheInitialized() && !InstanceData->TryPopulateIndexCache(*this, RigVMInstance, InternalResultName, RigVMFunctionHeader.Name))
+			if (!InstanceData->VariableIndexCache.IsIndexCacheInitialized() && !InstanceData->VariableIndexCache.TryPopulateIndexCache(InstanceData->ParamData, RigVMInstance, RigVMFunctionHeader.Name, InternalResultName))
 			{
 				return false;
 			}
 
 			// Prior to execution, copy over ST binding values to RigVM variables
 			FAnimNextGraphInstance& GraphInstance = AnimExecContext->GetRootGraphInstance();
-			TArrayView<const uint8> ArgumentIndexes = InstanceData->GetVMArgumentIndexes();
+			TArrayView<const uint8> ArgumentIndexes = InstanceData->VariableIndexCache.GetVMArgumentIndexes();
 			TConstArrayView<FPropertyBagPropertyDesc> StateTreePropertyDescs = InstanceData->ParamData.GetPropertyBagStruct()->GetPropertyDescs();
 			TConstArrayView<FPropertyBagPropertyDesc> AnimNextGraphPropertyDescs = GraphInstance.GetVariables().GetPropertyBagStruct()->GetPropertyDescs();
 			FStructView StateTreeParamData = InstanceData->ParamData.GetMutableValue();
@@ -198,7 +73,10 @@ bool FAnimNextStateTreeRigVMConditionBase::TestCondition(FStateTreeExecutionCont
 
 			bool bResult = false;
 			RigVMInstance->GetVM()->ExecuteVM(ExtendedExecuteContext, InternalEventName);
-			GraphInstance.GetVariable(InternalResultName, bResult);
+			if (ensure(ExtendedExecuteContext.ExternalVariableRuntimeData.IsValidIndex(InstanceData->VariableIndexCache.GetVMResultIndex())))
+			{
+				bResult = (bool)*ExtendedExecuteContext.ExternalVariableRuntimeData[InstanceData->VariableIndexCache.GetVMResultIndex()].Memory;
+			}
 
 			return bResult;
 		}
@@ -269,143 +147,22 @@ void FAnimNextStateTreeRigVMConditionBase::PostEditNodeChangeChainProperty(const
 	}
 }
 
-void FAnimNextStateTreeRigVMConditionBase::GetProgrammaticVariables(FAnimNextStateTreeProgrammaticVariableParams& InProgrammaticVariableParams, const UStateTreeState* State, const FStateTreeBindableStructDesc& Desc)
+void FAnimNextStateTreeRigVMConditionBase::GetProgrammaticFunctionHeaders(FAnimNextStateTreeProgrammaticFunctionHeaderParams& InProgrammaticFunctionHeaderParams, const UStateTreeState* State, const FStateTreeBindableStructDesc& Desc)
 {
-	UAnimNextDataInterface_EditorData* EditorData = InProgrammaticVariableParams.EditorData;
-	const FRigVMCompileSettings& Settings = InProgrammaticVariableParams.Settings;
-	FRigVMClient& RigVMClient = InProgrammaticVariableParams.RigVMClient;
-	TArray<FAnimNextProgrammaticVariable>& OutVariables = InProgrammaticVariableParams.OutVariables;
+	using namespace UE::AnimNext::UncookedOnly;
+
+	FAnimNextGetFunctionHeaderCompileContext& OutCompileContext = InProgrammaticFunctionHeaderParams.OutCompileContext;
 	
 	StateName = State->Name;
 	NodeId = Desc.ID;
-	InternalResultName = MakeInternalVariableName(ResultName);
-	InternalEventName = MakeInternalEventName(RigVMFunctionHeader.Name);
-
-	// Create input variables & result variable
-	{
-		for (const FRigVMGraphFunctionArgument& Argument : RigVMFunctionHeader.Arguments)
-		{
-			FRigVMGraphFunctionArgument InternallyNamedArgument = Argument;
-			InternallyNamedArgument.Name = MakeInternalVariableName(Argument.Name);
-			OutVariables.Add(FAnimNextProgrammaticVariable::FromRigVMGraphFunctionArgument(InternallyNamedArgument));
-		}
-	}
-}
-
-void FAnimNextStateTreeRigVMConditionBase::GetProgrammaticGraphs(FAnimNextStateTreeProgrammaticGraphParams& InProgrammaticGraphParams, const UStateTreeState* State, const FStateTreeBindableStructDesc& Desc)
-{
-	UAnimNextDataInterface_EditorData* EditorData = InProgrammaticGraphParams.EditorData;
-	const FRigVMCompileSettings& Settings = InProgrammaticGraphParams.Settings;
-	FRigVMClient& RigVMClient = InProgrammaticGraphParams.RigVMClient;
-	TArray<URigVMGraph*>& OutGraphs = InProgrammaticGraphParams.OutGraphs;
-
-	URigVMGraph* Graph = NewObject<URigVMGraph>(EditorData, RigVMFunctionHeader.Name, RF_Transient);
-	Graph->SetSchemaClass(UAnimNextAnimationGraphSchema::StaticClass());
-
-	URigVMController* Controller = RigVMClient.GetOrCreateController(Graph);
-	Controller->RemoveNodes(Controller->GetGraph()->GetNodes());
-
-	// Create RigVM graph for function
-	{
-		FString FunctionNodeString = RigVMFunctionHeader.Name.ToString();
-		bool bSetupUndoRedo = false;
-
-		// Add event entry node, so VM can call our event
-		URigVMUnitNode* EventNode = Controller->AddUnitNode(FRigVMFunction_UserDefinedEvent::StaticStruct(), TEXT("Execute"), FVector2D::ZeroVector, FunctionNodeString, bSetupUndoRedo);
-		if (!EventNode)
-		{
-			Settings.ReportError(FString::Printf(TEXT("Failed to add Event Node to VM Condition: %s"), *RigVMFunctionHeader.Name.ToString()));
-			return;
-		}
-		URigVMPin* EventNamePin = EventNode->FindPin(GET_MEMBER_NAME_STRING_CHECKED(FRigVMFunction_UserDefinedEvent, EventName));
-		if (EventNamePin == nullptr)
-		{
-			Settings.ReportError(FString::Printf(TEXT("Failed to add Event Name Pin to VM Condition: %s"), *RigVMFunctionHeader.Name.ToString()));
-			return;
-		}
-		FString EventNameString = InternalEventName.ToString();
-		Controller->SetPinDefaultValue(EventNamePin->GetPinPath(), EventNameString);
-
-		// Add function node
-		URigVMLibraryNode* LibraryNode = Cast<URigVMLibraryNode>(RigVMFunctionHeader.LibraryPointer.GetNodeSoftPath().ResolveObject());
-		if (!LibraryNode)
-		{
-			Settings.ReportError(FString::Printf(TEXT("Failed to get Library Node for VM Condition: %s"), *RigVMFunctionHeader.Name.ToString()));
-			return;
-		}
-		URigVMFunctionReferenceNode* FunctionNode = Controller->AddFunctionReferenceNode(LibraryNode, FVector2D::ZeroVector, FunctionNodeString, bSetupUndoRedo);
-		if (!FunctionNode)
-		{
-			Settings.ReportError(FString::Printf(TEXT("Failed to add Library Node to VM Condition: %s"), *RigVMFunctionHeader.Name.ToString()));
-			return;
-		}
-
-		auto GetNameSafe = [](const UObject* InObject)
-		{
-			if (InObject)
-			{
-				return InObject->GetFName().ToString();
-			}
-
-			return FString("Null Obj");
-		};
-
-		// Generate & link input arguments, also generate result variable node but link later
-		URigVMVariableNode* ResultNode = nullptr;
-		for (const FRigVMGraphFunctionArgument& Argument : RigVMFunctionHeader.Arguments)
-		{
-			bool bIsGetter = Argument.Name != ResultName;
-			FName InternalVariableName = MakeInternalVariableName(Argument.Name);
-			URigVMVariableNode* FunctionArgumentVariableNode = Controller->AddVariableNode(InternalVariableName, Argument.CPPType.ToString(), Argument.CPPTypeObject.Get(), bIsGetter, Argument.DefaultValue, FVector2D::ZeroVector, InternalVariableName.ToString(), bSetupUndoRedo);
-			if (!FunctionArgumentVariableNode)
-			{
-				Settings.ReportError(FString::Printf(TEXT("Failed to add internal variable node for VM Condition: %s, var: %s"), *RigVMFunctionHeader.Name.ToString(), *InternalVariableName.ToString()));
-				return;
-			}
-
-			if (bIsGetter)
-			{
-				URigVMPin* ArgumentValuePin = FunctionArgumentVariableNode->GetValuePin();
-				URigVMPin* FunctionArgumentPin = FunctionNode->FindPin(Argument.Name.ToString());
-
-				if (!Controller->AddLink(ArgumentValuePin, FunctionArgumentPin, bSetupUndoRedo))
-				{
-					Settings.ReportError(FString::Printf(TEXT("Failed to link pins for VM Condition: %s -> %s"), *GetNameSafe(ArgumentValuePin), *GetNameSafe(FunctionArgumentPin)));
-					return;
-				}
-			}
-			else
-			{
-				ResultNode = FunctionArgumentVariableNode;
-			}
-		}
-		if (!ResultNode)
-		{
-			Settings.ReportError(FString::Printf(TEXT("Failed to add internal result node for VM Condition: %s, result: %s"), *RigVMFunctionHeader.Name.ToString(), *InternalResultName.ToString()));
-			return;
-		}
-
-		// Link Execute pins
-		URigVMPin* ExecuteOutputPin = EventNode->FindPin(FRigVMStruct::ExecuteContextName.ToString());
-		URigVMPin* ExecuteInputPin = ResultNode->FindPin(FRigVMStruct::ExecuteContextName.ToString());
-		if (!Controller->AddLink(ExecuteOutputPin, ExecuteInputPin, bSetupUndoRedo))
-		{
-			Settings.ReportError(FString::Printf(TEXT("Failed to link execute pins for VM Condition: %s -> %s"), *GetNameSafe(ExecuteOutputPin), *GetNameSafe(ExecuteInputPin)));
-			return;
-		}
-
-		// Link Result pins
-		FString ResultNodeString = ResultName.ToString();
-		URigVMPin* FunctionResultPin = FunctionNode->FindPin(ResultNodeString);
-		URigVMPin* ResultValuePin = ResultNode->GetValuePin();
-		if (!Controller->AddLink(FunctionResultPin, ResultValuePin, bSetupUndoRedo))
-		{
-			Settings.ReportError(FString::Printf(TEXT("Failed to link result pins for VM Condition: %s -> %s"), *GetNameSafe(FunctionResultPin), *GetNameSafe(ResultValuePin)));
-			return;
-		}
-
-		OutGraphs.Add(Graph);
-	}
+	InternalResultName = FName(FUtils::MakeFunctionWrapperVariableName(RigVMFunctionHeader.Name, ResultName));
+	InternalEventName = FName(FUtils::MakeFunctionWrapperEventName(RigVMFunctionHeader.Name));
+
+	FAnimNextProgrammaticFunctionHeader AnimNextFunctionHeader = {};
+	AnimNextFunctionHeader.Wrapped = RigVMFunctionHeader;
+	AnimNextFunctionHeader.bGenerateParamVariables = true;
+	AnimNextFunctionHeader.bGenerateReturnVariables = true;
+	OutCompileContext.GetMutableFunctionHeaders().Add(AnimNextFunctionHeader);
 }
 
 #endif // WITH_EDITOR
diff --git a/Engine/Plugins/Experimental/AnimNextStateTree/Source/AnimNextStateTree/Private/RigVMUtils/RigVariableIndexCache.cpp b/Engine/Plugins/Experimental/AnimNextStateTree/Source/AnimNextStateTree/Private/RigVMUtils/RigVariableIndexCache.cpp
new file mode 100644
--- /dev/null
+++ b/Engine/Plugins/Experimental/AnimNextStateTree/Source/AnimNextStateTree/Private/RigVMUtils/RigVariableIndexCache.cpp
@@ -0,0 +1,85 @@
+// Copyright Epic Games, Inc. All Rights Reserved.
+
+#include "RigVMUtils/RigVariableIndexCache.h"
+
+#include "AnimNextExecuteContext.h"
+#include "Graph/AnimNextGraphContextData.h"
+#include "Graph/AnimNextGraphInstance.h"
+#include "StructUtils/PropertyBag.h"
+#include "TraitCore/ExecutionContext.h"
+
+FString FRigVariableIndexCache::MakeFunctionWrapperVariableName(FName InFunctionName, FName InVariableName)
+{
+	// We assume the function name is enough for variable name uniqueness in this graph (We don't yet desire global uniqueness).
+	return TEXT("__InternalVar_") + InFunctionName.ToString() + "_" + InVariableName.ToString();
+}
+
+FString FRigVariableIndexCache::MakeFunctionWrapperEventName(FName InFunctionName)
+{
+	return TEXT("__InternalCall_") + InFunctionName.ToString();
+}
+
+bool FRigVariableIndexCache::IsIndexCacheInitialized() const
+{
+	return bInitialized;
+}
+
+bool FRigVariableIndexCache::TryPopulateIndexCache(const FRigVMMemoryStorageStruct& ParamData, const UAnimNextDataInterface* InDataInterface, const FName InFunctionName, const FName InResultName)
+{
+	if (bInitialized)
+	{
+		return true;
+	}
+
+	ArgumentIndexes.Reserve(ParamData.Num());
+	TArray<FRigVMExternalVariable> ExternalVariables = InDataInterface->GetExternalVariables();
+	for (const FPropertyBagPropertyDesc& Desc : ParamData.GetPropertyBagStruct()->GetPropertyDescs())
+	{
+		for (int32 Index = 0; Index < ExternalVariables.Num(); Index++)
+		{
+			FName VariableName = ExternalVariables[Index].Name;
+			if (VariableName == FName(FRigVariableIndexCache::MakeFunctionWrapperVariableName(InFunctionName, Desc.Name)))
+			{
+				ArgumentIndexes.Add(Index);
+				break;
+			}
+		}
+	}
+
+	// Search all variables for result, it won't be part of the arg param data
+	bool bResultFound = InResultName == NAME_None ? true : false;
+	if (!bResultFound)
+	{
+		for (int32 Index = 0; Index < ExternalVariables.Num(); Index++)
+		{
+			FName VariableName = ExternalVariables[Index].Name;
+			if (VariableName == InResultName)
+			{
+				ResultIndex = Index;
+				bResultFound = true;
+				break;
+			}
+		}
+	}
+
+	bool bAllArgsFound = ArgumentIndexes.Num() == ParamData.Num();
+	if (bAllArgsFound && bResultFound)
+	{
+		bInitialized = true;
+	}
+
+	// Verify we init first time. We don't want to be spin-attempting to initialize.
+	verify(bInitialized);
+
+	return bInitialized;
+}
+
+TArrayView<const uint8> FRigVariableIndexCache::GetVMArgumentIndexes() const
+{
+	return MakeArrayView(ArgumentIndexes);
+}
+
+int32 FRigVariableIndexCache::GetVMResultIndex() const
+{
+	return ResultIndex;
+}
diff --git a/Engine/Plugins/Experimental/AnimNextStateTree/Source/AnimNextStateTree/Public/Conditions/AnimNextStateTreeRigVMConditionBase.h b/Engine/Plugins/Experimental/AnimNextStateTree/Source/AnimNextStateTree/Public/Conditions/AnimNextStateTreeRigVMConditionBase.h
--- a/Engine/Plugins/Experimental/AnimNextStateTree/Source/AnimNextStateTree/Public/Conditions/AnimNextStateTreeRigVMConditionBase.h
+++ b/Engine/Plugins/Experimental/AnimNextStateTree/Source/AnimNextStateTree/Public/Conditions/AnimNextStateTreeRigVMConditionBase.h
@@ -7,6 +7,7 @@
 #include "RigVMCore/RigVMExternalVariable.h"
 #include "RigVMCore/RigVMGraphFunctionDefinition.h"
 #include "RigVMCore/RigVMMemoryStorageStruct.h"
+#include "RigVMUtils/RigVariableIndexCache.h"
 #include "StateTreeConditionBase.h"
 #include "StateTreeReference.h"
 #include "Templates/SubclassOf.h"
@@ -18,8 +19,7 @@ class UAnimNextDataInterface_EditorData;
 class URigVMGraph;
 class UStateTreeState;
 
-struct FAnimNextStateTreeProgrammaticVariableParams;
-struct FAnimNextStateTreeProgrammaticGraphParams;
+struct FAnimNextStateTreeProgrammaticFunctionHeaderParams;
 struct FStateTreeBindableStructDesc;
 struct FStateTreeExecutionContext;
 struct FRigVMClient;
@@ -37,27 +37,6 @@ public:
 
 public:
 
-	/** True if VM variables are indexed */
-	bool IsIndexCacheInitialized() const;
-
-	/** True if successfully populated index cache from Data Interface, or already populated*/
-	bool TryPopulateIndexCache(const FAnimNextStateTreeRigVMConditionBase& OwningCondition, TObjectPtr<UAnimNextDataInterface> InDataInterface, FName InResultName, FName InFunctionName);
-
-	/** Gets Argument Indexes, you are responsible for checking if initialized first */
-	TArrayView<const uint8> GetVMArgumentIndexes() const;
-
-	/** Gets Result Index, you are responsible for checking if initialized first */
-	uint8 GetVMResultIndex() const;
-
-protected:
-
-	struct FRigVariableIndexCache
-	{
-		TArray<uint8, TInlineAllocator<8>> ArgumentIndexes = {};
-		uint8 ResultIndex = 0;
-		bool bInitialized = false;
-	};
-
 	/** Helper struct to store parent RigVM variable indicies, */
 	FRigVariableIndexCache VariableIndexCache = FRigVariableIndexCache();
 };
@@ -72,9 +51,6 @@ struct ANIMNEXTSTATETREE_API FAnimNextStateTreeRigVMConditionBase : public FStat
 
 	using FInstanceDataType = FAnimNextStateTreeRigVMConditionInstanceData;
 
-	static const FName InternalEventPrefixName;
-	static const FName InternalVariablePrefixName;
-
 public:
 
 	//~ Begin FStateTreeConditionBase Interface
@@ -91,19 +67,10 @@ public:
 	//~ End FStateTreeNodeBase Interface
 	
 #if WITH_EDITOR
-	/** Adds condition variables to parent StateTree RigVM */
-	virtual void GetProgrammaticVariables(FAnimNextStateTreeProgrammaticVariableParams& InProgrammaticVariableParams, const UStateTreeState* State, const FStateTreeBindableStructDesc& Desc);
-
-	/** Adds condition graphs to parent StateTree RigVM */
-	virtual void GetProgrammaticGraphs(FAnimNextStateTreeProgrammaticGraphParams& InProgrammaticGraphParams, const UStateTreeState* State, const FStateTreeBindableStructDesc& Desc);
+	/** Adds condition function headers to parent StateTree RigVM */
+	virtual void GetProgrammaticFunctionHeaders(FAnimNextStateTreeProgrammaticFunctionHeaderParams& InProgrammaticFunctionHeaderParams, const UStateTreeState* State, const FStateTreeBindableStructDesc& Desc);
 #endif
 
-	/** Generates a unique internal name for RigVM events, name not intended to be user-facing */
-	FName MakeInternalEventName(const FName InName) const;
-
-	/** Generates a unique internal name for RigVM variables, name not intended to be user-facing */
-	FName MakeInternalVariableName(const FName InName) const;
-
 public:
 
 #if WITH_EDITORONLY_DATA
diff --git a/Engine/Plugins/Experimental/AnimNextStateTree/Source/AnimNextStateTree/Public/RigVMUtils/RigVariableIndexCache.h b/Engine/Plugins/Experimental/AnimNextStateTree/Source/AnimNextStateTree/Public/RigVMUtils/RigVariableIndexCache.h
new file mode 100644
--- /dev/null
+++ b/Engine/Plugins/Experimental/AnimNextStateTree/Source/AnimNextStateTree/Public/RigVMUtils/RigVariableIndexCache.h
@@ -0,0 +1,48 @@
+// Copyright Epic Games, Inc. All Rights Reserved.
+
+#pragma once
+
+#include "CoreMinimal.h"
+#include "AnimNextStateTreeTypes.h"
+
+class UAnimNextDataInterface;
+struct FRigVMMemoryStorageStruct;
+
+/** 
+ * Helper struct to store parent RigVM variable indicies,.
+ * 
+ * This is needed as the internal variables we read / write will be offset somewhere after the public vars in the data interface.
+ */
+struct FRigVariableIndexCache
+{
+
+public:
+
+	// @TODO: Relocate, must be in sync with UncookedOnlyUtils.h
+	/** Make a variable name that we use as a wrapper for a function param or return */
+	static FString MakeFunctionWrapperVariableName(FName InFunctionName, FName InVariableName);
+
+	// @TODO: Relocate, must be in sync with UncookedOnlyUtils.h
+	/** Make an event name that we use as a wrapper to call RigVM functions */
+	static FString MakeFunctionWrapperEventName(FName InFunctionName);
+
+public:
+
+	/** True if VM variables are indexed */
+	bool IsIndexCacheInitialized() const;
+
+	/** True if successfully populated index cache from Data Interface, or already populated*/
+	bool TryPopulateIndexCache(const FRigVMMemoryStorageStruct& ParamData, const UAnimNextDataInterface* InDataInterface, const FName InFunctionName, const FName InResultName = NAME_None);
+
+	/** Gets Argument Indexes, you are responsible for checking if initialized first */
+	TArrayView<const uint8> GetVMArgumentIndexes() const;
+
+	/** Gets Result Index, you are responsible for checking if initialized first */
+	int32 GetVMResultIndex() const;
+
+protected:
+
+	TArray<uint8, TInlineAllocator<8>> ArgumentIndexes = {};
+	int32 ResultIndex = INDEX_NONE;
+	bool bInitialized = false;
+};
\ No newline at end of file
diff --git a/Engine/Plugins/Experimental/AnimNextStateTree/Source/AnimNextStateTreeUncookedOnly/Internal/AnimNextStateTree_EditorData.cpp b/Engine/Plugins/Experimental/AnimNextStateTree/Source/AnimNextStateTreeUncookedOnly/Internal/AnimNextStateTree_EditorData.cpp
--- a/Engine/Plugins/Experimental/AnimNextStateTree/Source/AnimNextStateTreeUncookedOnly/Internal/AnimNextStateTree_EditorData.cpp
+++ b/Engine/Plugins/Experimental/AnimNextStateTree/Source/AnimNextStateTreeUncookedOnly/Internal/AnimNextStateTree_EditorData.cpp
@@ -7,6 +7,7 @@
 #include "AnimNextStateTreeEditorOnlyTypes.h"
 #include "AnimNextStateTreeWorkspaceAssetUserData.h"
 #include "Conditions/AnimNextStateTreeRigVMConditionBase.h"
+#include "Compilation/AnimNextGetGraphCompileContext.h"
 #include "DataInterface/AnimNextDataInterface_EditorData.h"
 #include "StateTree.h"
 #include "StateTreeEditorData.h"
@@ -56,12 +57,14 @@ TConstArrayView<TSubclassOf<UAnimNextRigVMAssetEntry>> UAnimNextStateTree_Editor
 	return Classes;
 }
 
-void UAnimNextStateTree_EditorData::OnPreCompileGetProgrammaticVariables(const FRigVMCompileSettings& InSettings, TArray<FAnimNextProgrammaticVariable>& OutVariables)
+void UAnimNextStateTree_EditorData::OnPreCompileGetProgrammaticFunctionHeaders(const FRigVMCompileSettings& InSettings, FAnimNextGetFunctionHeaderCompileContext& OutCompileContext)
 {
+	Super::OnPreCompileGetProgrammaticFunctionHeaders(InSettings, OutCompileContext);
+
 	if (UAnimNextStateTree* AnimStateTree = UE::AnimNext::UncookedOnly::FUtils::GetAsset<UAnimNextStateTree>(this))
 	{
 		// Give child nodes a chance to add compile-time only variables
-		FAnimNextStateTreeProgrammaticVariableParams ProgrammaticVariableParams(this, InSettings, RigVMClient, OutVariables);
+		FAnimNextStateTreeProgrammaticFunctionHeaderParams ProgrammaticFunctionHeaderParams(this, InSettings, RigVMClient, OutCompileContext);
 
 		// Populate UID info on our nodes
 		auto MakeNodeVariables = [&](UStateTreeState& State, UStateTreeState* ParentState)
@@ -79,7 +82,7 @@ void UAnimNextStateTree_EditorData::OnPreCompileGetProgrammaticVariables(const F
 					Desc.DataSource = EStateTreeBindableStructSource::Condition;
 
 					// Graph gen isn't const since we need to cache generated names in the node
-					const_cast<FAnimNextStateTreeRigVMConditionBase*>(Condition)->GetProgrammaticVariables(ProgrammaticVariableParams, &State, Desc);
+					const_cast<FAnimNextStateTreeRigVMConditionBase*>(Condition)->GetProgrammaticFunctionHeaders(ProgrammaticFunctionHeaderParams, &State, Desc);
 				}
 			}
 
@@ -101,8 +104,10 @@ void UAnimNextStateTree_EditorData::OnPreCompileGetProgrammaticVariables(const F
 	}
 }
 
-void UAnimNextStateTree_EditorData::OnPreCompileGetProgrammaticGraphs(const FRigVMCompileSettings& InSettings, TArray<URigVMGraph*>& OutGraphs)
+void UAnimNextStateTree_EditorData::OnPreCompileGetProgrammaticGraphs(const FRigVMCompileSettings& InSettings, FAnimNextGetGraphCompileContext& OutCompileContext)
 {
+	Super::OnPreCompileGetProgrammaticGraphs(InSettings, OutCompileContext);
+
 	if (UAnimNextStateTree* AnimStateTree = UE::AnimNext::UncookedOnly::FUtils::GetAsset<UAnimNextStateTree>(this))
 	{
 		URigVMGraph* Graph = NewObject<URigVMGraph>(this, NAME_None, RF_Transient);
@@ -185,46 +190,7 @@ void UAnimNextStateTree_EditorData::OnPreCompileGetProgrammaticGraphs(const FRig
 			return;
 		}
 
-		// Give child nodes a chance to add graphs
-		FAnimNextStateTreeProgrammaticGraphParams ProgrammaticGraphParams(this, InSettings, RigVMClient, OutGraphs);
-
-		// Populate UID info on our nodes
-		auto MakeNodeGraphs = [&](UStateTreeState& State, UStateTreeState* ParentState)
-		{
-			for (const FStateTreeEditorNode& Node : State.EnterConditions)
-			{
-				// Manually check for each child node type, UStruct's do not reflect interfaces so we can't use an interface here.
-				if (const FAnimNextStateTreeRigVMConditionBase* Condition = Node.Node.GetPtr<FAnimNextStateTreeRigVMConditionBase>())
-				{
-					FStateTreeBindableStructDesc Desc;
-					Desc.StatePath = State.GetPath() + TEXT("/EnterConditions");
-					Desc.Struct = Condition->GetInstanceDataType();
-					Desc.Name = Node.GetName();
-					Desc.ID = Node.ID;
-					Desc.DataSource = EStateTreeBindableStructSource::Condition;
-
-					// Graph gen isn't const since we need to cache generated names in the node
-					const_cast<FAnimNextStateTreeRigVMConditionBase*>(Condition)->GetProgrammaticGraphs(ProgrammaticGraphParams, &State, Desc);
-				}
-			}
-
-			// @TODO: Add RigVM graph task support
-			//for (const FStateTreeEditorNode& Node : State.Tasks)
-			//{
-			//}
-
-			return EStateTreeVisitor::Continue;
-		};
-
-		if (AnimStateTree->StateTree)
-		{
-			if (UStateTreeEditorData* StateTreeEditorData = Cast<UStateTreeEditorData>(AnimStateTree->StateTree->EditorData))
-			{
-				StateTreeEditorData->VisitHierarchy(MakeNodeGraphs);
-			}
-		}
-
 		// @TODO: This has to be the last graph or ST will not execute. Figure out why, graph order should not be implicitly required
-		OutGraphs.Add(Graph);
+		OutCompileContext.GetMutableProgrammaticGraphs().Add(Graph);
 	}
 }
diff --git a/Engine/Plugins/Experimental/AnimNextStateTree/Source/AnimNextStateTreeUncookedOnly/Internal/AnimNextStateTree_EditorData.h b/Engine/Plugins/Experimental/AnimNextStateTree/Source/AnimNextStateTreeUncookedOnly/Internal/AnimNextStateTree_EditorData.h
--- a/Engine/Plugins/Experimental/AnimNextStateTree/Source/AnimNextStateTreeUncookedOnly/Internal/AnimNextStateTree_EditorData.h
+++ b/Engine/Plugins/Experimental/AnimNextStateTree/Source/AnimNextStateTreeUncookedOnly/Internal/AnimNextStateTree_EditorData.h
@@ -22,8 +22,8 @@ protected:
 	virtual TConstArrayView<TSubclassOf<UAnimNextRigVMAssetEntry>> GetEntryClasses() const override;
 
 	// Allows this asset to generate variables to be injected at compilation time, separate method to allow programmatic graphs to use these vars
-	virtual void OnPreCompileGetProgrammaticVariables(const FRigVMCompileSettings& InSettings, TArray<FAnimNextProgrammaticVariable>& OutVariables) override;
+	virtual void OnPreCompileGetProgrammaticFunctionHeaders(const FRigVMCompileSettings& InSettings, FAnimNextGetFunctionHeaderCompileContext& OutCompileContext) override;
 
 	// Allows this asset to generate graphs to be injected at compilation time
-	virtual void OnPreCompileGetProgrammaticGraphs(const FRigVMCompileSettings& InSettings, TArray<URigVMGraph*>& OutGraphs) override;
+	virtual void OnPreCompileGetProgrammaticGraphs(const FRigVMCompileSettings& InSettings, FAnimNextGetGraphCompileContext& OutCompileContext) override;
 };
\ No newline at end of file
