Fix ST not running on Windows Cooked due to comparing values when didn't need to.

diff --git a/Engine/Plugins/Experimental/AnimNextStateTree/Source/AnimNextStateTree/Private/AnimStateTreeTrait.cpp b/Engine/Plugins/Experimental/AnimNextStateTree/Source/AnimNextStateTree/Private/AnimStateTreeTrait.cpp
--- a/Engine/Plugins/Experimental/AnimNextStateTree/Source/AnimNextStateTree/Private/AnimStateTreeTrait.cpp
+++ b/Engine/Plugins/Experimental/AnimNextStateTree/Source/AnimNextStateTree/Private/AnimStateTreeTrait.cpp
@@ -101,7 +101,9 @@ namespace UE::AnimNext
 			StateTreeExecutionContext.SetContextDataByName(UStateTreeAnimNextSchema::AnimStateTreeExecutionContextName, FStateTreeDataView(FAnimNextStateTreeTraitContext::StaticStruct(), reinterpret_cast<uint8*>(&TraitContext)));
 
 			FAnimNextGraphInstance* OwnerGraphInstance = &Binding.GetTraitPtr().GetNodeInstance()->GetOwner();
-			if(OwnerGraphInstance && InstanceData->StateTree->GetDefaultParameters().Identical(&OwnerGraphInstance->GetVariables(), 0))
+			bool bPropertyLayoutMatches = FStructUtils::TheSameLayout(InstanceData->StateTree->GetDefaultParameters().GetPropertyBagStruct(), OwnerGraphInstance->GetVariables().GetPropertyBagStruct());
+
+			if(OwnerGraphInstance && ensure(bPropertyLayoutMatches))
 			{
 				const FInstancedPropertyBag& StateTreeParameters = InstanceData->StateTree->GetDefaultParameters();
 				FStructView MutableVariables = OwnerGraphInstance->GetMutableVariables().GetMutableValue();
@@ -144,6 +146,14 @@ namespace UE::AnimNext
 					StateTreeExecutionContext.Start(&OwnerGraphInstance->GetVariables());
 				}
 			}
+			else
+			{
+				if (StateTreeExecutionContext.IsValid())
+				{
+					UE_LOG(LogAnimation, Error, TEXT("Failed to start StateTree. PropertyLayoutMatch: %u"), bPropertyLayoutMatches);
+					StateTreeExecutionContext.Stop(EStateTreeRunStatus::Failed);
+				}
+			}
 			
 		}	
 	}
@@ -167,7 +177,7 @@ namespace UE::AnimNext
 			
 			FStateTreeExecutionContext StateTreeExecutionContext(*Owner, *InstanceData->StateTree, InstanceData->InstanceData);
 			FAnimNextGraphInstance* OwnerGraphInstance = &Binding.GetTraitPtr().GetNodeInstance()->GetOwner();
-			if(OwnerGraphInstance && InstanceData->StateTree->GetDefaultParameters().GetPropertyBagStruct() == OwnerGraphInstance->GetVariables().GetPropertyBagStruct())
+			if(OwnerGraphInstance && StateTreeExecutionContext.GetLastTickStatus() != EStateTreeRunStatus::Failed)
 			{
 				FAnimNextStateTreeTraitContext TraitContext(Context, &Binding);
 				StateTreeExecutionContext.SetContextDataByName(UStateTreeAnimNextSchema::AnimStateTreeExecutionContextName, FStateTreeDataView(FAnimNextStateTreeTraitContext::StaticStruct(), reinterpret_cast<uint8*>(&TraitContext)));
