Change AnimNextStateTreeGraphInstanceTask to copy payload struct instead of passing by StructView, because during blend out, the StructView would point to invalid memory when the State Tree Task instance goes out of scope

diff --git a/Engine/Plugins/Experimental/AnimNextStateTree/Source/AnimNextStateTree/Private/AnimNextStateTreeGraphInstanceTask.cpp b/Engine/Plugins/Experimental/AnimNextStateTree/Source/AnimNextStateTree/Private/AnimNextStateTreeGraphInstanceTask.cpp
--- a/Engine/Plugins/Experimental/AnimNextStateTree/Source/AnimNextStateTree/Private/AnimNextStateTreeGraphInstanceTask.cpp
+++ b/Engine/Plugins/Experimental/AnimNextStateTree/Source/AnimNextStateTree/Private/AnimNextStateTreeGraphInstanceTask.cpp
@@ -26,7 +26,9 @@ EStateTreeRunStatus FAnimNextStateTreeGraphInstanceTask::EnterState(FStateTreeEx
 
 	// Reference task's default payload
 	FAnimNextDataInterfacePayload Payload;
-	Payload.AddStructView(InstanceData.Payload.GetMutableValue());
+	FInstancedStruct PayloadStruct;
+	PayloadStruct.InitializeAs(InstanceData.Payload.GetPropertyBagStruct(), InstanceData.Payload.GetValue().GetMemory());
+	Payload.AddNative(MoveTemp(PayloadStruct));
 
 	// Temporarily disabling parameter overriding because it is blocking data interface variables from being accessible in subgraphs
 	// FInstancedStruct NewNativePayload;
