Fix ensure in AnimNext StateTree due to different default values in propertybag.

diff --git a/Engine/Plugins/Experimental/AnimNextStateTree/Source/AnimNextStateTree/Private/AnimStateTreeTrait.cpp b/Engine/Plugins/Experimental/AnimNextStateTree/Source/AnimNextStateTree/Private/AnimStateTreeTrait.cpp
--- a/Engine/Plugins/Experimental/AnimNextStateTree/Source/AnimNextStateTree/Private/AnimStateTreeTrait.cpp
+++ b/Engine/Plugins/Experimental/AnimNextStateTree/Source/AnimNextStateTree/Private/AnimStateTreeTrait.cpp
@@ -101,7 +101,7 @@ namespace UE::AnimNext
 			StateTreeExecutionContext.SetContextDataByName(UStateTreeAnimNextSchema::AnimStateTreeExecutionContextName, FStateTreeDataView(FAnimNextStateTreeTraitContext::StaticStruct(), reinterpret_cast<uint8*>(&TraitContext)));
 
 			FAnimNextGraphInstance* OwnerGraphInstance = &Binding.GetTraitPtr().GetNodeInstance()->GetOwner();
-			bool bPropertyLayoutMatches = FStructUtils::TheSameLayout(InstanceData->StateTree->GetDefaultParameters().GetPropertyBagStruct(), OwnerGraphInstance->GetVariables().GetPropertyBagStruct());
+			bool bPropertyLayoutMatches = InstanceData->StateTree->GetDefaultParameters().HasSameLayout(OwnerGraphInstance->GetVariables());
 
 			if(OwnerGraphInstance && ensure(bPropertyLayoutMatches))
 			{
diff --git a/Engine/Plugins/Runtime/PropertyBindingUtils/Source/PropertyBindingUtils/Private/PropertyBindingBindingCollection.cpp b/Engine/Plugins/Runtime/PropertyBindingUtils/Source/PropertyBindingUtils/Private/PropertyBindingBindingCollection.cpp
--- a/Engine/Plugins/Runtime/PropertyBindingUtils/Source/PropertyBindingUtils/Private/PropertyBindingBindingCollection.cpp
+++ b/Engine/Plugins/Runtime/PropertyBindingUtils/Source/PropertyBindingUtils/Private/PropertyBindingBindingCollection.cpp
@@ -1080,7 +1080,7 @@ bool FPropertyBindingBindingCollection::CopyProperty(const FPropertyBindingCopyI
 	
 	if (SourceStructView.IsValid() && TargetStructView.IsValid())
 	{
-		check(SourceStructView.GetStruct() == Copy.SourceStructType
+		check(FStructUtils::TheSameLayout(SourceStructView.GetStruct(), Copy.SourceStructType) 
 			|| (SourceStructView.GetStruct() && SourceStructView.GetStruct()->IsChildOf(Copy.SourceStructType)));
 			
 		uint8* SourceAddress = GetAddress(SourceStructView, Copy.SourceIndirection, Copy.SourceLeafProperty);
diff --git a/Engine/Plugins/Runtime/StateTree/Source/StateTreeModule/Private/StateTreeExecutionContext.cpp b/Engine/Plugins/Runtime/StateTree/Source/StateTreeModule/Private/StateTreeExecutionContext.cpp
--- a/Engine/Plugins/Runtime/StateTree/Source/StateTreeModule/Private/StateTreeExecutionContext.cpp
+++ b/Engine/Plugins/Runtime/StateTree/Source/StateTreeModule/Private/StateTreeExecutionContext.cpp
@@ -2347,7 +2347,7 @@ FStateTreeIndex16 FStateTreeExecutionContext::CollectExternalData(const UStateTr
 
 bool FStateTreeExecutionContext::SetGlobalParameters(const FInstancedPropertyBag& Parameters)
 {
-	if (ensureMsgf(RootStateTree.GetDefaultParameters().GetPropertyBagStruct() == Parameters.GetPropertyBagStruct(),
+	if (ensureMsgf(RootStateTree.GetDefaultParameters().HasSameLayout(Parameters),
 		TEXT("Parameters must be of the same struct type. Make sure to migrate the provided parameters to the same type as the StateTree default parameters.")))
 	{
 		InstanceDataStorage->SetGlobalParameters(Parameters);
diff --git a/Engine/Source/Runtime/CoreUObject/Private/StructUtils/PropertyBag.cpp b/Engine/Source/Runtime/CoreUObject/Private/StructUtils/PropertyBag.cpp
--- a/Engine/Source/Runtime/CoreUObject/Private/StructUtils/PropertyBag.cpp
+++ b/Engine/Source/Runtime/CoreUObject/Private/StructUtils/PropertyBag.cpp
@@ -1850,6 +1850,15 @@ bool FInstancedPropertyBag::OwnsPropertyDesc(const FPropertyBagPropertyDesc& Des
 	return false;
 }
 
+bool FInstancedPropertyBag::HasSameLayout(const FInstancedPropertyBag& Other) const
+{
+	const UPropertyBag* BagStruct = GetPropertyBagStruct();
+	const UPropertyBag* OtherBagStruct = Other.GetPropertyBagStruct();
+
+	// Note: FStructUtils::TheSameLayout will check nullptr & type per property, not just layout. We do not check name however
+	return FStructUtils::TheSameLayout(BagStruct, OtherBagStruct);
+}
+
 const void* FInstancedPropertyBag::GetValueAddress(const FPropertyBagPropertyDesc* Desc) const
 {
 	if (Desc == nullptr || !Value.IsValid())
diff --git a/Engine/Source/Runtime/CoreUObject/Public/StructUtils/PropertyBag.h b/Engine/Source/Runtime/CoreUObject/Public/StructUtils/PropertyBag.h
--- a/Engine/Source/Runtime/CoreUObject/Public/StructUtils/PropertyBag.h
+++ b/Engine/Source/Runtime/CoreUObject/Public/StructUtils/PropertyBag.h
@@ -483,6 +483,9 @@ struct COREUOBJECT_API FInstancedPropertyBag
 	/** Returns true if we own the supplied property description, false otherwise. */
 	bool OwnsPropertyDesc(const FPropertyBagPropertyDesc& Desc) const;
 
+	/** Returns true if our propertybag has the same layout & per property types as another propertybag. */
+	bool HasSameLayout(const FInstancedPropertyBag& Other) const;
+
 	/** @return const view to the struct that holds the values. NOTE: The returned value/view cannot be serialized, use this to access the struct only temporarily. */
 	FConstStructView GetValue() const { return Value; };
 
