Add AnimNextStateTreeRigVMTask

diff --git a/Engine/Plugins/Experimental/AnimNextAnimGraph/Source/AnimNextAnimGraph/Public/Graph/AnimNextGraphInstance.h b/Engine/Plugins/Experimental/AnimNextAnimGraph/Source/AnimNextAnimGraph/Public/Graph/AnimNextGraphInstance.h
--- a/Engine/Plugins/Experimental/AnimNextAnimGraph/Source/AnimNextAnimGraph/Public/Graph/AnimNextGraphInstance.h
+++ b/Engine/Plugins/Experimental/AnimNextAnimGraph/Source/AnimNextAnimGraph/Public/Graph/AnimNextGraphInstance.h
@@ -25,6 +25,7 @@ struct FAnimNextModuleInjectionComponent;
 struct FRigUnit_AnimNextRunAnimationGraph_v1;
 struct FRigUnit_AnimNextRunAnimationGraph_v2;
 struct FAnimNextStateTreeRigVMConditionBase;
+struct FAnimNextStateTreeRigVMTaskBase;
 
 namespace UE::AnimNext
 {
@@ -253,6 +254,7 @@ private:
 	friend FRigUnit_AnimNextRunAnimationGraph_v2;
 	friend UE::AnimNext::FInjectionUtils;
 	friend FAnimNextStateTreeRigVMConditionBase;	// Temp - remove: Needs to access mutable variables to copy-in data for function shim
+	friend FAnimNextStateTreeRigVMTaskBase;	// Temp - remove: Needs to access mutable variables to copy-in data for function shim
 	friend UE::AnimNext::FBlendSpacePlayerTrait;
 };
 
diff --git a/Engine/Plugins/Experimental/AnimNextAnimGraph/Source/AnimNextAnimGraphUncookedOnly/Private/Graph/AnimNextAnimationGraph_EditorData.cpp b/Engine/Plugins/Experimental/AnimNextAnimGraph/Source/AnimNextAnimGraphUncookedOnly/Private/Graph/AnimNextAnimationGraph_EditorData.cpp
--- a/Engine/Plugins/Experimental/AnimNextAnimGraph/Source/AnimNextAnimGraphUncookedOnly/Private/Graph/AnimNextAnimationGraph_EditorData.cpp
+++ b/Engine/Plugins/Experimental/AnimNextAnimGraph/Source/AnimNextAnimGraphUncookedOnly/Private/Graph/AnimNextAnimationGraph_EditorData.cpp
@@ -971,6 +971,12 @@ void UAnimNextAnimationGraph_EditorData::OnPreCompileGetProgrammaticGraphs(const
 				// Generate & link input arguments, also generate result variable node but link later
 				for (const FRigVMGraphFunctionArgument& Argument : FunctionHeader.Arguments)
 				{
+					// Execution context is captured as arg pins, skip those for internval variable gen
+					if (Argument.Direction == ERigVMPinDirection::IO)
+					{
+						continue;
+					}
+
 					bool bIsGetter = Argument.Direction == ERigVMPinDirection::Input;
 
 					if (bIsGetter && AnimNextFunctionHeader.bGenerateParamVariables)
diff --git a/Engine/Plugins/Experimental/AnimNextStateTree/Source/AnimNextStateTree/Private/Tasks/AnimNextStateTreeRigVMTaskBase.cpp b/Engine/Plugins/Experimental/AnimNextStateTree/Source/AnimNextStateTree/Private/Tasks/AnimNextStateTreeRigVMTaskBase.cpp
new file mode 100644
--- /dev/null
+++ b/Engine/Plugins/Experimental/AnimNextStateTree/Source/AnimNextStateTree/Private/Tasks/AnimNextStateTreeRigVMTaskBase.cpp
@@ -0,0 +1,175 @@
+// Copyright Epic Games, Inc. All Rights Reserved.
+
+#include "Tasks/AnimNextStateTreeRigVMTaskBase.h"
+#include "AnimNextExecuteContext.h"
+#include "Graph/AnimNextGraphContextData.h"
+#include "Graph/AnimNextGraphInstance.h"
+#include "StateTreeLinker.h"
+#include "StateTreeExecutionContext.h"
+#include "StructUtils/PropertyBag.h"
+#include "TraitCore/ExecutionContext.h"
+
+#if WITH_EDITOR
+#include "AnimNextStateTreeEditorOnlyTypes.h"
+#include "RigVMBlueprintGeneratedClass.h"
+#include "RigVMFunctions/RigVMFunctionDefines.h"
+#include "StateTree.h"
+#include "StateTreeEditorData.h"
+#include "UncookedOnlyUtils.h"
+#endif // WITH_EDITOR
+
+#define LOCTEXT_NAMESPACE "AnimNextStateTreeRigVMTaskBase"
+
+#include UE_INLINE_GENERATED_CPP_BY_NAME(AnimNextStateTreeRigVMTaskBase)
+
+////////////////////////////////////////////////////////////////////////////////////////////////////////////////////////
+// FAnimNextStateTreeRigVMTaskBase
+
+bool FAnimNextStateTreeRigVMTaskBase::Link(FStateTreeLinker& Linker)
+{
+	Linker.LinkExternalData(TraitContextHandle);
+	return true;
+}
+
+EStateTreeRunStatus FAnimNextStateTreeRigVMTaskBase::EnterState(FStateTreeExecutionContext& Context, const FStateTreeTransitionResult& Transition) const
+{
+	if (FAnimNextStateTreeRigVMTaskInstanceData* InstanceData = Context.GetInstanceDataPtr<FAnimNextStateTreeRigVMTaskInstanceData>(*this))
+	{
+		FAnimNextStateTreeTraitContext& ExecContext = Context.GetExternalData(TraitContextHandle);
+		const UE::AnimNext::FExecutionContext* AnimExecContext = ExecContext.GetAnimExecuteContext();
+
+		// Need const cast as VM execution isn't const.
+		if (UAnimNextDataInterface* RigVMInstance = const_cast<UAnimNextDataInterface*>(AnimExecContext->GetRootGraphInstance().GetDataInterface()))
+		{
+			// @TODO: For now don't pass result name. Since we use fire and forget
+			// Cache RigVM variable indexes, which can vary as they are parent derived
+			if (!InstanceData->VariableIndexCache.IsIndexCacheInitialized() && !InstanceData->VariableIndexCache.TryPopulateIndexCache(InstanceData->ParamData, RigVMInstance, RigVMFunctionHeader.Name, NAME_None))
+			{
+				return EStateTreeRunStatus::Failed;
+			}
+
+			// Prior to execution, copy over ST binding values to RigVM variables
+			FAnimNextGraphInstance& GraphInstance = AnimExecContext->GetRootGraphInstance();
+			TArrayView<const uint8> ArgumentIndexes = InstanceData->VariableIndexCache.GetVMArgumentIndexes();
+			TConstArrayView<FPropertyBagPropertyDesc> StateTreePropertyDescs = InstanceData->ParamData.GetPropertyBagStruct()->GetPropertyDescs();
+			TConstArrayView<FPropertyBagPropertyDesc> AnimNextGraphPropertyDescs = GraphInstance.GetVariables().GetPropertyBagStruct()->GetPropertyDescs();
+			FStructView StateTreeParamData = InstanceData->ParamData.GetMutableValue();
+
+			for (int32 Index = 0; Index < ArgumentIndexes.Num(); Index++)
+			{
+				uint8 RigVMExternalVariableIndex = ArgumentIndexes[Index];
+
+				const FPropertyBagPropertyDesc& StateTreePropertyDesc = StateTreePropertyDescs[Index];
+				void* SourceAddress = StateTreePropertyDesc.CachedProperty->ContainerPtrToValuePtr<void*>(StateTreeParamData.GetMemory());
+
+				const FPropertyBagPropertyDesc& AnimNextGraphDesc = AnimNextGraphPropertyDescs[RigVMExternalVariableIndex];
+				GraphInstance.GetMutableVariables().SetValue(AnimNextGraphDesc.Name, StateTreePropertyDesc.CachedProperty, StateTreeParamData.GetMemory());
+			}
+
+			FRigVMExtendedExecuteContext& ExtendedExecuteContext = GraphInstance.GetExtendedExecuteContext();
+			FAnimNextExecuteContext& AnimNextContext = GraphInstance.GetExtendedExecuteContext().GetPublicDataSafe<FAnimNextExecuteContext>();
+
+			FAnimNextGraphContextData ContextData(GraphInstance.GetModuleInstance(), &GraphInstance);
+			UE::AnimNext::FScopedExecuteContextData ContextDataScope(AnimNextContext, ContextData);
+
+			// @TODO: Add option to execute per tick?
+			RigVMInstance->GetVM()->ExecuteVM(ExtendedExecuteContext, InternalEventName);
+
+			// @TODO: For now we use a fire & forget model. But in the future we might want users to be able to bind to the result value from a function
+			// This will require a special UI customization, & standalone execute is useful for modifying existing vars. So save for later effort.
+			//if (ensure(ExtendedExecuteContext.ExternalVariableRuntimeData.IsValidIndex(InstanceData->VariableIndexCache.GetVMResultIndex())))
+			//{
+			//	auto Result = *ExtendedExecuteContext.ExternalVariableRuntimeData[InstanceData->VariableIndexCache.GetVMResultIndex()].Memory;
+			//}
+
+			return EStateTreeRunStatus::Running;
+		}
+	}
+
+	return EStateTreeRunStatus::Failed;
+}
+
+#if WITH_EDITOR
+FText FAnimNextStateTreeRigVMTaskBase::GetDescription(const FGuid& ID, FStateTreeDataView InstanceDataView, const IStateTreeBindingLookup& BindingLookup, EStateTreeNodeFormatting Formatting) const
+{
+	return LOCTEXT("AnimNextStateTreeConditon_Desc", "RigVM function driven Task");
+}
+
+void FAnimNextStateTreeRigVMTaskBase::PostEditNodeChangeChainProperty(const FPropertyChangedChainEvent& PropertyChangedEvent, FStateTreeDataView InstanceDataView)
+{
+	if (PropertyChangedEvent.GetMemberPropertyName() == GET_MEMBER_NAME_CHECKED(FAnimNextStateTreeRigVMTaskBase, TaskFunctionName))
+	{
+		// Function Name selected has changed. Update param struct & result / event names used during execution.
+		if (FAnimNextStateTreeRigVMTaskInstanceData* InstanceData = InstanceDataView.GetMutablePtr<FAnimNextStateTreeRigVMTaskInstanceData>())
+		{
+			// @TODO: This relies on the function name being unique (Ex: In a workspace). For now that's okay. Later on however we will want to use a more robust function picker
+			auto GetRigVMFunctionHeader = [](FName InName)
+			{
+				FRigVMGraphFunctionHeader Result = FRigVMGraphFunctionHeader();
+
+				TMap<FAssetData, FRigVMGraphFunctionHeaderArray> FunctionExports;
+				UE::AnimNext::UncookedOnly::FUtils::GetExportedFunctionsFromAssetRegistry(UE::AnimNext::AnimNextPublicGraphFunctionsExportsRegistryTag, FunctionExports);
+				UE::AnimNext::UncookedOnly::FUtils::GetExportedFunctionsFromAssetRegistry(UE::AnimNext::ControlRigAssetPublicGraphFunctionsExportsRegistryTag, FunctionExports);
+
+				for (const auto& Export : FunctionExports.Array())
+				{
+					for (const FRigVMGraphFunctionHeader& FunctionHeader : Export.Value.Headers)
+					{
+						if (InName == FunctionHeader.Name)
+						{
+							Result = FunctionHeader;
+							return Result;
+						}
+					}
+				}
+
+				return Result;
+			};
+
+			RigVMFunctionHeader = GetRigVMFunctionHeader(TaskFunctionName);
+			InstanceData->ParamData = FRigVMMemoryStorageStruct();
+
+			TArray<FRigVMPropertyDescription> RigVMPropertyDescriptions;
+			for (const FRigVMGraphFunctionArgument& Argument : RigVMFunctionHeader.Arguments)
+			{
+				if (Argument.Direction == ERigVMPinDirection::Input)
+				{
+					FString CppTypeString = Argument.CPPType.ToString();
+					UObject* CppTypeObject = Argument.CPPTypeObject.Get();
+					FRigVMPropertyDescription RigVMPropertyDescription = FRigVMPropertyDescription(Argument.Name, CppTypeString, CppTypeObject, Argument.DefaultValue);
+					RigVMPropertyDescriptions.Add(RigVMPropertyDescription);
+				}
+				else if (Argument.Direction == ERigVMPinDirection::Output)
+				{
+					ensure(Argument.CPPType == RigVMTypeUtils::BoolTypeName);
+					ResultName = Argument.Name;
+				}
+			}
+		
+			InstanceData->ParamData.AddProperties(RigVMPropertyDescriptions);
+		}
+	}
+}
+
+void FAnimNextStateTreeRigVMTaskBase::GetProgrammaticFunctionHeaders(FAnimNextStateTreeProgrammaticFunctionHeaderParams& InProgrammaticFunctionHeaderParams, const UStateTreeState* State, const FStateTreeBindableStructDesc& Desc)
+{
+	using namespace UE::AnimNext::UncookedOnly;
+
+	FAnimNextGetFunctionHeaderCompileContext& OutCompileContext = InProgrammaticFunctionHeaderParams.OutCompileContext;
+	
+	StateName = State->Name;
+	NodeId = Desc.ID;
+	InternalResultName = FName(FUtils::MakeFunctionWrapperVariableName(RigVMFunctionHeader.Name, ResultName));
+	InternalEventName = FName(FUtils::MakeFunctionWrapperEventName(RigVMFunctionHeader.Name));
+
+	FAnimNextProgrammaticFunctionHeader AnimNextFunctionHeader = {};
+	AnimNextFunctionHeader.Wrapped = RigVMFunctionHeader;
+	AnimNextFunctionHeader.bGenerateParamVariables = true;
+	// @TODO: For now don't gen return variables. Since we use fire and forget
+	AnimNextFunctionHeader.bGenerateReturnVariables = false;
+	OutCompileContext.GetMutableFunctionHeaders().Add(AnimNextFunctionHeader);
+}
+
+#endif // WITH_EDITOR
+
+#undef LOCTEXT_NAMESPACE
diff --git a/Engine/Plugins/Experimental/AnimNextStateTree/Source/AnimNextStateTree/Public/Tasks/AnimNextStateTreeRigVMTaskBase.h b/Engine/Plugins/Experimental/AnimNextStateTree/Source/AnimNextStateTree/Public/Tasks/AnimNextStateTreeRigVMTaskBase.h
new file mode 100644
--- /dev/null
+++ b/Engine/Plugins/Experimental/AnimNextStateTree/Source/AnimNextStateTree/Public/Tasks/AnimNextStateTreeRigVMTaskBase.h
@@ -0,0 +1,105 @@
+// Copyright Epic Games, Inc. All Rights Reserved.
+
+#pragma once
+
+#include "AnimNextStateTreeContext.h"
+#include "AnimNextStateTreeTypes.h"
+#include "DataInterface/AnimNextDataInterface.h"
+#include "RigVMCore/RigVMExternalVariable.h"
+#include "RigVMCore/RigVMGraphFunctionDefinition.h"
+#include "RigVMCore/RigVMMemoryStorageStruct.h"
+#include "RigVMUtils/RigVariableIndexCache.h"
+#include "StateTreeTaskBase.h"
+#include "StateTreeReference.h"
+#include "Templates/SubclassOf.h"
+
+#include "AnimNextStateTreeRigVMTaskBase.generated.h"
+
+class UAnimNextDataInterface;
+class UAnimNextDataInterface_EditorData;
+class URigVMGraph;
+class UStateTreeState;
+
+struct FAnimNextStateTreeProgrammaticFunctionHeaderParams;
+struct FStateTreeBindableStructDesc;
+struct FStateTreeExecutionContext;
+struct FRigVMClient;
+
+USTRUCT()
+struct ANIMNEXTSTATETREE_API FAnimNextStateTreeRigVMTaskInstanceData
+{
+	GENERATED_BODY()
+
+public:
+
+	/** Per instance copy of ST function arguments we will copy into our RigVM before execution */
+	UPROPERTY(EditAnywhere, Category = "Parameter", Meta = (FixedLayout, ShowOnlyInnerProperties))
+	FRigVMMemoryStorageStruct ParamData = FRigVMMemoryStorageStruct();
+
+public:
+
+	/** Helper struct to store parent RigVM variable indicies, */
+	FRigVariableIndexCache VariableIndexCache = FRigVariableIndexCache();
+};
+
+/**
+ * Wrapper for RigVM based Tasks. 
+ */
+USTRUCT()
+struct ANIMNEXTSTATETREE_API FAnimNextStateTreeRigVMTaskBase : public FAnimNextStateTreeTaskBase
+{
+	GENERATED_BODY()
+
+	using FInstanceDataType = FAnimNextStateTreeRigVMTaskInstanceData;
+
+public:
+
+	//~ Begin FStateTreeTaskBase Interface
+	virtual bool Link(FStateTreeLinker& Linker) override;
+	virtual EStateTreeRunStatus EnterState(FStateTreeExecutionContext& Context, const FStateTreeTransitionResult& Transition) const override;
+	//~ End FStateTreeTaskBase Interface
+
+	//~ Begin FStateTreeNodeBase Interface
+	virtual const UStruct* GetInstanceDataType() const override { return FInstanceDataType::StaticStruct(); }
+#if WITH_EDITOR
+	virtual FText GetDescription(const FGuid& ID, FStateTreeDataView InstanceDataView, const IStateTreeBindingLookup& BindingLookup, EStateTreeNodeFormatting Formatting = EStateTreeNodeFormatting::Text) const override;
+	virtual void PostEditNodeChangeChainProperty(const FPropertyChangedChainEvent& PropertyChangedEvent, FStateTreeDataView InstanceDataView) override;
+#endif
+	//~ End FStateTreeNodeBase Interface
+	
+#if WITH_EDITOR
+	/** Adds Task function headers to parent StateTree RigVM */
+	virtual void GetProgrammaticFunctionHeaders(FAnimNextStateTreeProgrammaticFunctionHeaderParams& InProgrammaticFunctionHeaderParams, const UStateTreeState* State, const FStateTreeBindableStructDesc& Desc);
+#endif
+
+public:
+
+#if WITH_EDITORONLY_DATA
+	UPROPERTY(EditAnywhere, Category = "Parameter", Meta = (GetOptions = "AnimNextStateTreeUncookedOnly.AnimNextStateTreeFunctionLibraryHelper.GetExposedAnimNextFunctionNames"))
+	FName TaskFunctionName = NAME_None;
+
+	/** Name of the Result of function */
+	UPROPERTY()
+	FName ResultName = NAME_None;
+#endif
+
+	UPROPERTY()
+	FRigVMGraphFunctionHeader RigVMFunctionHeader = FRigVMGraphFunctionHeader();
+
+	/** Owning state name, populated during programmatic graph creation */
+	UPROPERTY()
+	FName StateName = NAME_None;
+
+	/** External node ID defined by owning state tree, populated during programmatic graph creation */
+	UPROPERTY()
+	FGuid NodeId = FGuid();
+
+	UPROPERTY()
+	FName InternalEventName = NAME_None;
+
+	UPROPERTY()
+	FName InternalResultName = NAME_None;
+
+public:
+	TStateTreeExternalDataHandle<FAnimNextStateTreeTraitContext> TraitContextHandle;
+};
\ No newline at end of file
diff --git a/Engine/Plugins/Experimental/AnimNextStateTree/Source/AnimNextStateTreeUncookedOnly/Internal/AnimNextStateTree_EditorData.cpp b/Engine/Plugins/Experimental/AnimNextStateTree/Source/AnimNextStateTreeUncookedOnly/Internal/AnimNextStateTree_EditorData.cpp
--- a/Engine/Plugins/Experimental/AnimNextStateTree/Source/AnimNextStateTreeUncookedOnly/Internal/AnimNextStateTree_EditorData.cpp
+++ b/Engine/Plugins/Experimental/AnimNextStateTree/Source/AnimNextStateTreeUncookedOnly/Internal/AnimNextStateTree_EditorData.cpp
@@ -20,6 +20,7 @@
 #include "RigVMCompiler/RigVMCompiler.h"
 #include "RigVMModel/RigVMClient.h"
 #include "RigVMModel/RigVMGraph.h"
+#include "Tasks/AnimNextStateTreeRigVMTaskBase.h"
 #include "TraitCore/TraitRegistry.h"
 #include "TraitCore/Trait.h"
 #include "AnimStateTreeTrait.h"
@@ -86,10 +87,22 @@ void UAnimNextStateTree_EditorData::OnPreCompileGetProgrammaticFunctionHeaders(c
 				}
 			}
 
-			// @TODO: Add RigVM graph task support
-			//for (const FStateTreeEditorNode& Node : State.Tasks)
-			//{
-			//}
+			for (const FStateTreeEditorNode& Node : State.Tasks)
+			{
+				// Manually check for each child node type, UStruct's do not reflect interfaces so we can't use an interface here.
+				if (const FAnimNextStateTreeRigVMTaskBase* Condition = Node.Node.GetPtr<FAnimNextStateTreeRigVMTaskBase>())
+				{
+					FStateTreeBindableStructDesc Desc;
+					Desc.StatePath = State.GetPath() + TEXT("/Tasks");
+					Desc.Struct = Condition->GetInstanceDataType();
+					Desc.Name = Node.GetName();
+					Desc.ID = Node.ID;
+					Desc.DataSource = EStateTreeBindableStructSource::Task;
+
+					// Graph gen isn't const since we need to cache generated names in the node
+					const_cast<FAnimNextStateTreeRigVMTaskBase*>(Condition)->GetProgrammaticFunctionHeaders(ProgrammaticFunctionHeaderParams, &State, Desc);
+				}
+			}
 
 			return EStateTreeVisitor::Continue;
 		};
