(AnimNext)StateTree changes:
- Added new EStateTreeDataSourceType type --> ExternalGlobalParameterData
- Added ability to populate FPropertyBindingCopyInfo to raw memory ptr mapping on execution context, used for accessing ExternalGlobalParameterData
- Deprecated public access to X, replaced by accessors to standalone propertybag + updated all usages
- Added ability to hide Parameters category from StateTree details customization (enabled for AnimNextStateTree)
- Added AnimNext speicifc StateTreeEditorData implementation, this maps/exposes AnimNext property bag as external global parameter data
- Updated AnimNextStateTreeTrait to use the above to prevent memory-copies

diff --git a/Engine/Plugins/Experimental/AnimNext/Source/AnimNextUncookedOnly/Internal/AnimNextRigVMAssetEditorData.h b/Engine/Plugins/Experimental/AnimNext/Source/AnimNextUncookedOnly/Internal/AnimNextRigVMAssetEditorData.h
--- a/Engine/Plugins/Experimental/AnimNext/Source/AnimNextUncookedOnly/Internal/AnimNextRigVMAssetEditorData.h
+++ b/Engine/Plugins/Experimental/AnimNext/Source/AnimNextUncookedOnly/Internal/AnimNextRigVMAssetEditorData.h
@@ -204,6 +204,7 @@ protected:
 	friend UAnimNextWorkspaceEditorMode;
 	friend UE::AnimNext::UncookedOnly::FScopedCompilerResults;
 	friend UE::AnimNext::Editor::FAssetCompilationHandler;
+	friend class UAnimNextStateTreeTreeEditorData;
 
 	// UObject interface
 	virtual void Serialize(FArchive& Ar) override;
diff --git a/Engine/Plugins/Experimental/AnimNextStateTree/Source/AnimNextStateTree/Internal/AnimStateTreeTrait.h b/Engine/Plugins/Experimental/AnimNextStateTree/Source/AnimNextStateTree/Internal/AnimStateTreeTrait.h
--- a/Engine/Plugins/Experimental/AnimNextStateTree/Source/AnimNextStateTree/Internal/AnimStateTreeTrait.h
+++ b/Engine/Plugins/Experimental/AnimNextStateTree/Source/AnimNextStateTree/Internal/AnimStateTreeTrait.h
@@ -10,6 +10,7 @@
 #include "TraitInterfaces/IUpdate.h"
 
 #include "StateTree.h"
+#include "StateTreeExecutionContext.h"
 #include "StateTreeReference.h"
 
 #include "AnimStateTreeTrait.generated.h"
@@ -46,6 +47,7 @@ struct FStateTreeTrait : FAdditiveTrait, IUpdate, IGarbageCollection
 	{
 		TObjectPtr<const UStateTree> StateTree;
 		FStateTreeInstanceData InstanceData;
+		FStateTreeExecutionContext::FExternalGlobalParameters StateTreeExternalParameters;
 
 		void Construct(const FExecutionContext& Context, const FTraitBinding& Binding);
 		void Destruct(const FExecutionContext& Context, const FTraitBinding& Binding);
diff --git a/Engine/Plugins/Experimental/AnimNextStateTree/Source/AnimNextStateTree/Private/AnimNextStateTreeSchema.cpp b/Engine/Plugins/Experimental/AnimNextStateTree/Source/AnimNextStateTree/Private/AnimNextStateTreeSchema.cpp
--- a/Engine/Plugins/Experimental/AnimNextStateTree/Source/AnimNextStateTree/Private/AnimNextStateTreeSchema.cpp
+++ b/Engine/Plugins/Experimental/AnimNextStateTree/Source/AnimNextStateTree/Private/AnimNextStateTreeSchema.cpp
@@ -19,6 +19,9 @@ const FName UStateTreeAnimNextSchema::AnimStateTreeExecutionContextName("Executi
 UStateTreeAnimNextSchema::UStateTreeAnimNextSchema() : ContextDataDescs(
 	{{ AnimStateTreeExecutionContextName, FAnimNextStateTreeTraitContext::StaticStruct(), FGuid(0xDFB93B7F, 0xEDBE4906, 0x851C66B2, 0x7585FA21) }})
 {
+#if WITH_EDITOR
+	DetailsCustomizationSettings.bShowStateTreeParameters = false;
+#endif // WITH_EDITOR
 }
 
 bool UStateTreeAnimNextSchema::IsStructAllowed(const UScriptStruct* InScriptStruct) const
@@ -41,4 +44,4 @@ bool UStateTreeAnimNextSchema::IsClassAllowed(const UClass* InClass) const
 bool UStateTreeAnimNextSchema::IsExternalItemAllowed(const UStruct& InStruct) const
 {
 	return InStruct.IsChildOf<FAnimNextStateTreeTraitContext>();
-}
+}
\ No newline at end of file
diff --git a/Engine/Plugins/Experimental/AnimNextStateTree/Source/AnimNextStateTree/Private/AnimStateTreeTrait.cpp b/Engine/Plugins/Experimental/AnimNextStateTree/Source/AnimNextStateTree/Private/AnimStateTreeTrait.cpp
--- a/Engine/Plugins/Experimental/AnimNextStateTree/Source/AnimNextStateTree/Private/AnimStateTreeTrait.cpp
+++ b/Engine/Plugins/Experimental/AnimNextStateTree/Source/AnimNextStateTree/Private/AnimStateTreeTrait.cpp
@@ -25,6 +25,9 @@ namespace UE::AnimNext
 	GENERATE_ANIM_TRAIT_IMPLEMENTATION(FStateTreeTrait, TRAIT_INTERFACE_ENUMERATOR, NULL_ANIM_TRAIT_INTERFACE_ENUMERATOR, NULL_ANIM_TRAIT_EVENT_ENUMERATOR)
 	#undef TRAIT_INTERFACE_ENUMERATOR
 
+#if ENABLE_ANIM_DEBUG 
+	TAutoConsoleVariable<bool> CVarLogPropertyBindingMemoryPtrInfo(TEXT("a.StateTree.LogPropertyBindingMemoryPtrInfo"), false, TEXT("Log information while generating property binding memory pointer mappings between RigVM and StateTree"));
+#endif //ENABLE_ANIM_DEBUG 
 
 	void FStateTreeTrait::FInstanceData::Construct(const FExecutionContext& Context, const FTraitBinding& Binding)
 	{
@@ -34,43 +37,42 @@ namespace UE::AnimNext
 
 	void FStateTreeTrait::FInstanceData::Destruct(const FExecutionContext& Context, const FTraitBinding& Binding)
 	{
+#if WITH_STATETREE_DEBUG
 		// @TODO: UE-240683
 		// We should call stop, but at the moment we don't have a good time / way to do that 
+		/*if (StateTree)
+		{
+			UObject* Owner = GetTransientPackage();
+			
+			FStateTreeExecutionContext StateTreeExecutionContext(*Owner, *StateTree, InstanceData);
+			FAnimNextGraphInstance* OwnerGraphInstance = &Binding.GetTraitPtr().GetNodeInstance()->GetOwner();
+			if(OwnerGraphInstance && StateTreeExecutionContext.GetLastTickStatus() != EStateTreeRunStatus::Failed)
+			{
+				FAnimNextStateTreeTraitContext TraitContext(*const_cast<FExecutionContext*>(&Context), &Binding);
+				StateTreeExecutionContext.SetContextDataByName(UStateTreeAnimNextSchema::AnimStateTreeExecutionContextName, FStateTreeDataView(FAnimNextStateTreeTraitContext::StaticStruct(), reinterpret_cast<uint8*>(&TraitContext)));
+				StateTreeExecutionContext.SetExternalGlobalParameters(&StateTreeExternalParameters);
+
+				StateTreeExecutionContext.SetCollectExternalDataCallback(FOnCollectStateTreeExternalData::CreateLambda([this, &TraitContext]( const FStateTreeExecutionContext& Context, const UStateTree* InStateTree, TArrayView<const FStateTreeExternalDataDesc> ExternalDataDescs, TArrayView<FStateTreeDataView> OutDataViews) -> bool
+				{
+					for (int32 Index = 0; Index < ExternalDataDescs.Num(); Index++)
+					{
+						const FStateTreeExternalDataDesc& ItemDesc = ExternalDataDescs[Index];
+						if (ItemDesc.Struct != nullptr)
+						{
+							if (ItemDesc.Struct->IsChildOf(FAnimNextStateTreeTraitContext::StaticStruct()))
+							{
+								OutDataViews[Index] = FStateTreeDataView(FAnimNextStateTreeTraitContext::StaticStruct(), reinterpret_cast<uint8*>(&TraitContext));
+							}
+						}
+					}
 		
-//		// Use Transient as owner, since our owning component is already dis-associated by now and all we need is the ST + Instance to stop.
-//		UObject* Owner = GetTransientPackage();
-//
-//		if (StateTree && false)
-//		{
-//			FStateTreeExecutionContext StateTreeExecutionContext(*Owner, *StateTree, InstanceData);
-//#if ENABLE_ANIM_DEBUG 
-//			// Clear owner, since at this point we are destructing and don't care about the context.
-//			// Also since the real owner may be GC'd by now (shutdown, level swap, etc).
-//			InstanceData.GetRuntimeValidation().SetContext(nullptr, StateTree);
-//#endif // ENABLE_ANIM_DEBUG 
-//			FAnimNextStateTreeTraitContext TraitContext(const_cast<FExecutionContext&>(Context), &Binding);
-//
-//			// @TODO: In separate CL make a helper method to do this. 
-//			StateTreeExecutionContext.SetCollectExternalDataCallback(FOnCollectStateTreeExternalData::CreateLambda([this, &TraitContext]( const FStateTreeExecutionContext& Context, const UStateTree* StateTree, TArrayView<const FStateTreeExternalDataDesc> ExternalDataDescs, TArrayView<FStateTreeDataView> OutDataViews) -> bool
-//			{
-//				for (int32 Index = 0; Index < ExternalDataDescs.Num(); Index++)
-//				{
-//					const FStateTreeExternalDataDesc& ItemDesc = ExternalDataDescs[Index];
-//					if (ItemDesc.Struct != nullptr)
-//					{
-//						if (ItemDesc.Struct->IsChildOf(FAnimNextStateTreeTraitContext::StaticStruct()))
-//						{
-//							OutDataViews[Index] = FStateTreeDataView(FAnimNextStateTreeTraitContext::StaticStruct(), reinterpret_cast<uint8*>(&TraitContext));
-//						}
-//					}
-//				}
-//
-//				return true;
-//			}));
-//
-//			StateTreeExecutionContext.Stop();
-//		}
+					return true;
+				}));
 
+				StateTreeExecutionContext.Stop();
+			}			
+		}*/	
+#endif
 		FTrait::FInstanceData::Destruct(Context, Binding);
 		IGarbageCollection::UnregisterWithGC(Context, Binding);
 	}
@@ -101,28 +103,74 @@ namespace UE::AnimNext
 			StateTreeExecutionContext.SetContextDataByName(UStateTreeAnimNextSchema::AnimStateTreeExecutionContextName, FStateTreeDataView(FAnimNextStateTreeTraitContext::StaticStruct(), reinterpret_cast<uint8*>(&TraitContext)));
 
 			FAnimNextGraphInstance* OwnerGraphInstance = &Binding.GetTraitPtr().GetNodeInstance()->GetOwner();
-			bool bPropertyLayoutMatches = InstanceData->StateTree->GetDefaultParameters().HasSameLayout(OwnerGraphInstance->GetVariables());
+			const bool bPropertyLayoutMatches = InstanceData->StateTree->GetDefaultParameters().HasSameLayout(OwnerGraphInstance->GetVariables());
 
 			if(OwnerGraphInstance && ensure(bPropertyLayoutMatches))
 			{
 				const FInstancedPropertyBag& StateTreeParameters = InstanceData->StateTree->GetDefaultParameters();
 				FStructView MutableVariables = OwnerGraphInstance->GetMutableVariables().GetMutableValue();
+				InstanceData->StateTreeExternalParameters.Reset();
 				
+				TArray<int32> StartPropertyOffsets;
 				const FRigVMExtendedExecuteContext& ExtendedExecuteContext = OwnerGraphInstance->GetExtendedExecuteContext();
 
-				// This copy-behaviour is temporary until we find a better way to directly patch the StateTree property binding/copies
 				const int32 NumVariables = StateTreeParameters.GetNumPropertiesInBag();
 				for(int32 VariableIndex = 0; VariableIndex < NumVariables; ++VariableIndex)
 				{
 					if (ExtendedExecuteContext.ExternalVariableRuntimeData.IsValidIndex(VariableIndex) && MutableVariables.GetMemory())
 					{
 						const FPropertyBagPropertyDesc& Desc = StateTreeParameters.GetPropertyBagStruct()->GetPropertyDescs()[VariableIndex];
-						void* TargetAddress = Desc.CachedProperty->ContainerPtrToValuePtr<void*>(MutableVariables.GetMemory());
-						void const* SourceAddress = ExtendedExecuteContext.ExternalVariableRuntimeData[VariableIndex].Memory;
+						StartPropertyOffsets.Add(Desc.CachedProperty->GetOffset_ForInternal());
+					}
+				}
+
+				for (const FPropertyBindingCopyInfoBatch& Batch : InstanceData->StateTree->GetPropertyBindings().GetCopyBatches())
+				{
+					for (const FPropertyBindingCopyInfo& Copy : InstanceData->StateTree->GetPropertyBindings().Super::GetBatchCopies(Batch))
+					{
+						const FStateTreeDataHandle& Handle = Copy.SourceDataHandle.Get<FStateTreeDataHandle>();
+						if (Handle.GetSource() == EStateTreeDataSourceType::ExternalGlobalParameterData)
+						{
+
+							const int32 RequiredOffset = [&Copy]()
+							{
+								if(Copy.SourceIndirection.Type == EPropertyBindingPropertyAccessType::Offset)
+								{
+									return (int32)Copy.SourceIndirection.Offset;
+								}
+
+								checkf(false, TEXT("Only expecting offset indirections for remapping"));
+								return Copy.SourceLeafProperty->GetOffset_ForInternal();
+							}();
+
+							uint8* MemoryPtr = nullptr;
+							const int32 NumOffsets = StartPropertyOffsets.Num();
+							for (int32 Index = 0; Index < NumOffsets; ++Index)
+							{
+								const int32 NextOffset = (Index < (NumOffsets - 1)) ? StartPropertyOffsets[Index + 1] : INDEX_NONE;
+								if (RequiredOffset >= StartPropertyOffsets[Index] && (RequiredOffset < NextOffset || NextOffset == INDEX_NONE))
+								{
+									// Incorporate the UPropertyBag root-level property offset into the remapped memory-ptr itself, so that the property-access indirection works "as-normal"
+									MemoryPtr = (uint8*)ExtendedExecuteContext.ExternalVariableRuntimeData[Index].Memory - StartPropertyOffsets[Index];
+									break;
+								}
+							}
+							check(MemoryPtr != nullptr);
 
-						Desc.CachedProperty->CopyCompleteValue(TargetAddress, SourceAddress);
+							if (InstanceData->StateTreeExternalParameters.Add(Copy, MemoryPtr))
+							{
+#if ENABLE_ANIM_DEBUG 
+								UE_CLOG(CVarLogPropertyBindingMemoryPtrInfo.GetValueOnAnyThread(), LogAnimation, Warning, TEXT("Mapped: Source: %s\nTarget: %sSize: %i\nOffset: %i\nType: %s"), *Copy.SourceLeafProperty->GetName(), *Copy.TargetLeafProperty->GetName(), Copy.CopySize, RequiredOffset, *FindObject<UEnum>(nullptr, TEXT("/Script/PropertyBindingUtils.EPropertyCopyType"))->GetNameStringByValue(static_cast<int64>(Copy.Type)));
+							}
+							else
+							{
+								UE_CLOG(CVarLogPropertyBindingMemoryPtrInfo.GetValueOnAnyThread(), LogAnimation, Warning, TEXT("Skipped: Source: %s\nTarget: %sSize: %i\nOffset: %i\nType: %s"), *Copy.SourceLeafProperty->GetName(), *Copy.TargetLeafProperty->GetName(), Copy.CopySize, RequiredOffset, *FindObject<UEnum>(nullptr, TEXT("/Script/PropertyBindingUtils.EPropertyCopyType"))->GetNameStringByValue(static_cast<int64>(Copy.Type)));
+#endif // ENABLE_ANIM_DEBUG 
+					        }
+				        }
 					}
 				}
+				StateTreeExecutionContext.SetExternalGlobalParameters(&InstanceData->StateTreeExternalParameters);
 				
 				StateTreeExecutionContext.SetCollectExternalDataCallback(FOnCollectStateTreeExternalData::CreateLambda([this, &TraitContext]( const FStateTreeExecutionContext& Context, const UStateTree* StateTree, TArrayView<const FStateTreeExternalDataDesc> ExternalDataDescs, TArrayView<FStateTreeDataView> OutDataViews) -> bool
 				{
@@ -143,7 +191,7 @@ namespace UE::AnimNext
 		
 				if (StateTreeExecutionContext.IsValid())
 				{
-					StateTreeExecutionContext.Start(&OwnerGraphInstance->GetVariables());
+					StateTreeExecutionContext.Start();
 				}
 			}
 			else
@@ -154,7 +202,6 @@ namespace UE::AnimNext
 					StateTreeExecutionContext.Stop(EStateTreeRunStatus::Failed);
 				}
 			}
-			
 		}	
 	}
 
@@ -181,24 +228,7 @@ namespace UE::AnimNext
 			{
 				FAnimNextStateTreeTraitContext TraitContext(Context, &Binding);
 				StateTreeExecutionContext.SetContextDataByName(UStateTreeAnimNextSchema::AnimStateTreeExecutionContextName, FStateTreeDataView(FAnimNextStateTreeTraitContext::StaticStruct(), reinterpret_cast<uint8*>(&TraitContext)));
-
-				const FInstancedPropertyBag& StateTreeParameters = InstanceData->StateTree->GetDefaultParameters();
-				FStructView MutableVariables = InstanceData->InstanceData.GetMutableStorage().GetMutableGlobalParameters();
-				
-				const FRigVMExtendedExecuteContext& ExtendedExecuteContext = OwnerGraphInstance->GetExtendedExecuteContext();
-
-				// This copy-behaviour is temporary until we find a better way to directly patch the StateTree property binding/copies
-				const int32 NumVariables = StateTreeParameters.GetNumPropertiesInBag();
-				for(int32 VariableIndex = 0; VariableIndex < NumVariables; ++VariableIndex)
-				{
-					if (ExtendedExecuteContext.ExternalVariableRuntimeData.IsValidIndex(VariableIndex) && MutableVariables.GetMemory())
-					{
-						const FPropertyBagPropertyDesc& Desc = StateTreeParameters.GetPropertyBagStruct()->GetPropertyDescs()[VariableIndex];
-						void* TargetAddress = Desc.CachedProperty->ContainerPtrToValuePtr<void*>(MutableVariables.GetMemory());
-						void const* SourceAddress = ExtendedExecuteContext.ExternalVariableRuntimeData[VariableIndex].Memory;
-						Desc.CachedProperty->CopyCompleteValue(TargetAddress, SourceAddress);
-					}
-				}
+				StateTreeExecutionContext.SetExternalGlobalParameters(&InstanceData->StateTreeExternalParameters);
 
 				StateTreeExecutionContext.SetCollectExternalDataCallback(FOnCollectStateTreeExternalData::CreateLambda([this, &TraitContext]( const FStateTreeExecutionContext& Context, const UStateTree* StateTree, TArrayView<const FStateTreeExternalDataDesc> ExternalDataDescs, TArrayView<FStateTreeDataView> OutDataViews) -> bool
 				{
diff --git a/Engine/Plugins/Experimental/AnimNextStateTree/Source/AnimNextStateTree/Public/AnimNextStateTree.h b/Engine/Plugins/Experimental/AnimNextStateTree/Source/AnimNextStateTree/Public/AnimNextStateTree.h
--- a/Engine/Plugins/Experimental/AnimNextStateTree/Source/AnimNextStateTree/Public/AnimNextStateTree.h
+++ b/Engine/Plugins/Experimental/AnimNextStateTree/Source/AnimNextStateTree/Public/AnimNextStateTree.h
@@ -41,6 +41,7 @@ class ANIMNEXTSTATETREE_API UAnimNextStateTree : public UAnimNextAnimationGraph
 
 public:
 	friend class UAnimNextStateTree_EditorData;
+	friend class UAnimNextStateTreeTreeEditorData;
 	friend class UAnimNextStateTreeFactory;
 
 	UPROPERTY()
diff --git a/Engine/Plugins/Experimental/AnimNextStateTree/Source/AnimNextStateTree/Public/AnimNextStateTreeContext.h b/Engine/Plugins/Experimental/AnimNextStateTree/Source/AnimNextStateTree/Public/AnimNextStateTreeContext.h
--- a/Engine/Plugins/Experimental/AnimNextStateTree/Source/AnimNextStateTree/Public/AnimNextStateTreeContext.h
+++ b/Engine/Plugins/Experimental/AnimNextStateTree/Source/AnimNextStateTree/Public/AnimNextStateTreeContext.h
@@ -2,6 +2,7 @@
 
 #pragma once
 
+#include "AnimStateTreeTrait.h"
 #include "StructUtils/InstancedStruct.h"
 
 #include "AnimNextStateTreeContext.generated.h"
@@ -23,6 +24,7 @@ struct FAnimNextStateTreeTraitContext
 	GENERATED_BODY()
 	
 	friend UE::AnimNext::FStateTreeTrait;
+	friend UE::AnimNext::FStateTreeTrait::FInstanceData;
 	
 	FAnimNextStateTreeTraitContext() {}
 
diff --git a/Engine/Plugins/Experimental/AnimNextStateTree/Source/AnimNextStateTree/Public/AnimNextStateTreeSchema.h b/Engine/Plugins/Experimental/AnimNextStateTree/Source/AnimNextStateTree/Public/AnimNextStateTreeSchema.h
--- a/Engine/Plugins/Experimental/AnimNextStateTree/Source/AnimNextStateTree/Public/AnimNextStateTreeSchema.h
+++ b/Engine/Plugins/Experimental/AnimNextStateTree/Source/AnimNextStateTree/Public/AnimNextStateTreeSchema.h
@@ -21,8 +21,19 @@ public:
 	virtual bool IsExternalItemAllowed(const UStruct& InStruct) const override;
 	
 	virtual TConstArrayView<FStateTreeExternalDataDesc> GetContextDataDescs() const override { return ContextDataDescs; }
+	virtual EGlobalParameterDataType GetGlobalParameterDataType() const override { return EGlobalParameterDataType::ExternalGlobalParameterData; }
+	
+#if WITH_EDITOR
+	virtual bool AllowUtilityConsiderations() const override { return false; }
+	virtual bool AllowEvaluators() const override { return false; }
+	virtual const FStateTreeDetailsCustomizationSettings& GetDetailsCustomizationSettings() const override { return DetailsCustomizationSettings; }
+#endif // WITH_EDITOR
 protected:
 	/** List of named external data required by schema and provided to the state tree through the execution context. */
 	UPROPERTY()
 	TArray<FStateTreeExternalDataDesc> ContextDataDescs;
+
+#if WITH_EDITOR
+	FStateTreeDetailsCustomizationSettings DetailsCustomizationSettings;
+#endif // WITH_EDITOR
 };
diff --git a/Engine/Plugins/Experimental/AnimNextStateTree/Source/AnimNextStateTreeEditor/Private/AnimNextStateTreeEditorModule.cpp b/Engine/Plugins/Experimental/AnimNextStateTree/Source/AnimNextStateTreeEditor/Private/AnimNextStateTreeEditorModule.cpp
--- a/Engine/Plugins/Experimental/AnimNextStateTree/Source/AnimNextStateTreeEditor/Private/AnimNextStateTreeEditorModule.cpp
+++ b/Engine/Plugins/Experimental/AnimNextStateTree/Source/AnimNextStateTreeEditor/Private/AnimNextStateTreeEditorModule.cpp
@@ -23,7 +23,10 @@
 #include "ToolMenus.h"
 #include "Framework/MultiBox/ToolMenuBase.h"
 #include "AnimNextRigVMAssetEditorData.h"
+#include "AnimNextStateTreeEditorData.h"
+#include "AnimNextStateTreeSchema.h"
 #include "StateTreeAssetCompilationHandler.h"
+#include "StateTreeEditorModule.h"
 
 #define LOCTEXT_NAMESPACE "AnimNextStateTreeEditorModule"
 
@@ -318,6 +321,11 @@ void FAnimNextStateTreeEditorModule::StartupModule()
 	TSharedPtr<FStateTreeAssetOutlinerItemDetails> StateItemDetails = MakeShareable<FStateTreeAssetOutlinerItemDetails>(new FStateTreeAssetOutlinerItemDetails());
 	WorkspaceEditorModule.RegisterWorkspaceItemDetails(Workspace::FOutlinerItemDetailsId(FAnimNextStateTreeOutlinerData::StaticStruct()->GetFName()), StaticCastSharedPtr<UE::Workspace::IWorkspaceOutlinerItemDetails>(StateItemDetails));	
 	WorkspaceEditorModule.RegisterWorkspaceItemDetails(Workspace::FOutlinerItemDetailsId(FAnimNextStateTreeStateOutlinerData::StaticStruct()->GetFName()), StaticCastSharedPtr<UE::Workspace::IWorkspaceOutlinerItemDetails>(StateItemDetails));
+	
+	if(FStateTreeEditorModule* StateTreeEditorModule = FStateTreeEditorModule::GetModulePtr())
+	{
+		StateTreeEditorModule->RegisterEditorDataClass(UStateTreeAnimNextSchema::StaticClass(), UAnimNextStateTreeTreeEditorData::StaticClass());		
+	}
 }
 
 void FAnimNextStateTreeEditorModule::ShutdownModule()
@@ -336,7 +344,12 @@ void FAnimNextStateTreeEditorModule::ShutdownModule()
 			AnimNextEditorModule->RemoveWorkspaceSupportedAssetClass(UStateTree::StaticClass()->GetClassPathName());
 			AnimNextEditorModule->UnregisterAssetCompilationHandler(UAnimNextStateTree::StaticClass()->GetClassPathName());
 		}
-	}
+
+		if(FStateTreeEditorModule* StateTreeEditorModule = FStateTreeEditorModule::GetModulePtr())
+		{
+			StateTreeEditorModule->UnregisterEditorDataClass(UStateTreeAnimNextSchema::StaticClass());			
+		}
+    }
 }
 
 IMPLEMENT_MODULE(FAnimNextStateTreeEditorModule, AnimNextStateTreeEditor)
diff --git a/Engine/Plugins/Experimental/AnimNextStateTree/Source/AnimNextStateTreeUncookedOnly/Internal/AnimNextStateTreeEditorData.cpp b/Engine/Plugins/Experimental/AnimNextStateTree/Source/AnimNextStateTreeUncookedOnly/Internal/AnimNextStateTreeEditorData.cpp
new file mode 100644
--- /dev/null
+++ b/Engine/Plugins/Experimental/AnimNextStateTree/Source/AnimNextStateTreeUncookedOnly/Internal/AnimNextStateTreeEditorData.cpp
@@ -0,0 +1,73 @@
+// Copyright Epic Games, Inc. All Rights Reserved.
+
+#pragma once
+
+#include "AnimNextStateTreeEditorData.h"
+
+#include "UncookedOnlyUtils.h"
+#include "AnimNextStateTree_EditorData.h"
+#include "StateTreeDelegates.h"
+
+void UAnimNextStateTreeTreeEditorData::PostLoad()
+{
+	Super::PostLoad();
+	if(const UAnimNextStateTree* OuterAnimNextStateTree = GetTypedOuter<UAnimNextStateTree>())
+	{
+		UAnimNextStateTree_EditorData* AnimNextEditorData = UE::AnimNext::UncookedOnly::FUtils::GetEditorData<UAnimNextStateTree_EditorData>(OuterAnimNextStateTree);	
+		AnimNextEditorData->ModifiedDelegate.AddUObject(this, &UAnimNextStateTreeTreeEditorData::HandleStateTreeAssetChanges);
+	}
+}
+
+void UAnimNextStateTreeTreeEditorData::CreateRootProperties(TArrayView<UE::PropertyBinding::FPropertyCreationDescriptor> InOutCreationDescs)
+{
+	UStateTree* OuterStateTree = CastChecked<UStateTree>(GetOuter());
+	UAnimNextStateTree* OuterAnimNextStateTree = CastChecked<UAnimNextStateTree>(OuterStateTree->GetOuter());
+
+	UAnimNextStateTree_EditorData* AnimNextEditorData = UE::AnimNext::UncookedOnly::FUtils::GetEditorData<UAnimNextStateTree_EditorData>(OuterAnimNextStateTree);
+
+	TArray<FPropertyBagPropertyDesc, TInlineAllocator<1>> PropertyDescs;
+	PropertyDescs.Reserve(InOutCreationDescs.Num());
+
+	// Generate unique names for the incoming property descs to avoid changing the existing properties in the bag
+	for (UE::PropertyBinding::FPropertyCreationDescriptor& CreationDesc : InOutCreationDescs)
+	{
+		int32 Index = CreationDesc.PropertyDesc.Name.GetNumber();
+		while (OuterAnimNextStateTree->VariableDefaults.FindPropertyDescByName(CreationDesc.PropertyDesc.Name))
+		{
+			CreationDesc.PropertyDesc.Name = FName(CreationDesc.PropertyDesc.Name, Index++);
+		}
+
+		// Try and export default value from incoming property desc + data
+		FString DefaultValue;
+		if (CreationDesc.SourceProperty && CreationDesc.SourceContainerAddress)
+		{
+			void const* SourceAddress = CreationDesc.SourceProperty->ContainerPtrToValuePtr<void>(CreationDesc.SourceContainerAddress);
+			CreationDesc.SourceProperty->ExportText_Direct(DefaultValue, SourceAddress, SourceAddress, nullptr, PPF_None);
+		}
+		
+		const FAnimNextParamType ParamType(CreationDesc.PropertyDesc.ValueType, CreationDesc.PropertyDesc.ContainerTypes.GetFirstContainerType(), CreationDesc.PropertyDesc.ValueTypeObject);
+		AnimNextEditorData->AddVariable(CreationDesc.PropertyDesc.Name, ParamType, *DefaultValue);
+	}
+}
+
+void UAnimNextStateTreeTreeEditorData::HandleStateTreeAssetChanges(UAnimNextRigVMAssetEditorData* InEditorData, EAnimNextEditorDataNotifType InType,
+	UObject* InSubject) const
+{	
+	switch(InType)
+	{
+	case EAnimNextEditorDataNotifType::UndoRedo:
+	case EAnimNextEditorDataNotifType::EntryAdded:
+	case EAnimNextEditorDataNotifType::EntryRemoved:
+	case EAnimNextEditorDataNotifType::EntryRenamed:
+	case EAnimNextEditorDataNotifType::EntryAccessSpecifierChanged:
+	case EAnimNextEditorDataNotifType::VariableTypeChanged:
+	case EAnimNextEditorDataNotifType::VariableDefaultValueChanged:
+		{			
+			const UStateTree* StateTree = GetTypedOuter<UStateTree>();
+			UE::StateTree::Delegates::OnParametersChanged.Broadcast(*StateTree);	
+		}		
+		break;
+	default:
+		break;
+	}
+}
diff --git a/Engine/Plugins/Experimental/AnimNextStateTree/Source/AnimNextStateTreeUncookedOnly/Internal/AnimNextStateTree_EditorData.cpp b/Engine/Plugins/Experimental/AnimNextStateTree/Source/AnimNextStateTreeUncookedOnly/Internal/AnimNextStateTree_EditorData.cpp
--- a/Engine/Plugins/Experimental/AnimNextStateTree/Source/AnimNextStateTreeUncookedOnly/Internal/AnimNextStateTree_EditorData.cpp
+++ b/Engine/Plugins/Experimental/AnimNextStateTree/Source/AnimNextStateTreeUncookedOnly/Internal/AnimNextStateTree_EditorData.cpp
@@ -33,19 +33,6 @@ TSubclassOf<UAssetUserData> UAnimNextStateTree_EditorData::GetAssetUserDataClass
 	return UAnimNextStateTreeWorkspaceAssetUserData::StaticClass();
 }
 
-void UAnimNextStateTree_EditorData::RecompileVM()
-{
-	UAnimNextAnimationGraph_EditorData::RecompileVM();
-
-	UAnimNextStateTree* AnimationStateTree = UE::AnimNext::UncookedOnly::FUtils::GetAsset<UAnimNextStateTree>(this);
-
-	UStateTree* InnerStateTree = AnimationStateTree->StateTree;
-
-	UStateTreeEditorData* InnerEditorData = Cast<UStateTreeEditorData>(InnerStateTree->EditorData);
-	InnerEditorData->RootParameters.ResetParametersAndOverrides();
-	InnerEditorData->RootParameters.Parameters = AnimationStateTree->VariableDefaults;
-}
-
 TConstArrayView<TSubclassOf<UAnimNextRigVMAssetEntry>> UAnimNextStateTree_EditorData::GetEntryClasses() const
 {
 	static const TSubclassOf<UAnimNextRigVMAssetEntry> Classes[] =
diff --git a/Engine/Plugins/Experimental/AnimNextStateTree/Source/AnimNextStateTreeUncookedOnly/Internal/AnimNextStateTree_EditorData.h b/Engine/Plugins/Experimental/AnimNextStateTree/Source/AnimNextStateTreeUncookedOnly/Internal/AnimNextStateTree_EditorData.h
--- a/Engine/Plugins/Experimental/AnimNextStateTree/Source/AnimNextStateTreeUncookedOnly/Internal/AnimNextStateTree_EditorData.h
+++ b/Engine/Plugins/Experimental/AnimNextStateTree/Source/AnimNextStateTreeUncookedOnly/Internal/AnimNextStateTree_EditorData.h
@@ -15,9 +15,7 @@ protected:
 	virtual TSubclassOf<UAssetUserData> GetAssetUserDataClass() const override;
 
 	friend class UAnimNextStateTreeFactory;
-	// IRigVMClientHost interface
-	virtual void RecompileVM() override;
-	
+
 	// UAnimNextRigVMAssetEditorData interface
 	virtual TConstArrayView<TSubclassOf<UAnimNextRigVMAssetEntry>> GetEntryClasses() const override;
 
diff --git a/Engine/Plugins/Experimental/AnimNextStateTree/Source/AnimNextStateTreeUncookedOnly/Public/AnimNextStateTreeEditorData.h b/Engine/Plugins/Experimental/AnimNextStateTree/Source/AnimNextStateTreeUncookedOnly/Public/AnimNextStateTreeEditorData.h
new file mode 100644
--- /dev/null
+++ b/Engine/Plugins/Experimental/AnimNextStateTree/Source/AnimNextStateTreeUncookedOnly/Public/AnimNextStateTreeEditorData.h
@@ -0,0 +1,36 @@
+// Copyright Epic Games, Inc. All Rights Reserved.
+
+#pragma once
+
+#include "AnimNextStateTree.h"
+#include "StateTreeEditorData.h"
+#include "UncookedOnlyUtils.h"
+
+#include "AnimNextStateTreeEditorData.generated.h"
+
+enum class EAnimNextEditorDataNotifType : uint8;
+
+UCLASS()
+class ANIMNEXTSTATETREEUNCOOKEDONLY_API UAnimNextStateTreeTreeEditorData : public UStateTreeEditorData
+{
+	GENERATED_BODY()
+
+public:
+	//~ Begin UStateTreeEditorData overrides
+	virtual void PostLoad() override;
+	
+	virtual const FInstancedPropertyBag& GetRootParametersPropertyBag() const override
+	{
+		const UStateTree* OuterStateTree = CastChecked<UStateTree>(GetOuter());
+		const UAnimNextStateTree* OuterAnimNextStateTree = CastChecked<UAnimNextStateTree>(OuterStateTree->GetOuter());
+		return OuterAnimNextStateTree->VariableDefaults;
+	}
+
+	virtual void CreateRootProperties(TArrayView<UE::PropertyBinding::FPropertyCreationDescriptor> InOutCreationDescs) override;
+	//~ End UStateTreeEditorData overrides
+
+protected:
+	void HandleStateTreeAssetChanges(UAnimNextRigVMAssetEditorData* InEditorData, EAnimNextEditorDataNotifType InType, UObject* InSubject) const;
+
+	friend class UAnimNextStateTree_EditorData;
+};
\ No newline at end of file
diff --git a/Engine/Plugins/Runtime/StateTree/Source/StateTreeEditorModule/Private/AsyncStateTreeDiff.cpp b/Engine/Plugins/Runtime/StateTree/Source/StateTreeEditorModule/Private/AsyncStateTreeDiff.cpp
--- a/Engine/Plugins/Runtime/StateTree/Source/StateTreeEditorModule/Private/AsyncStateTreeDiff.cpp
+++ b/Engine/Plugins/Runtime/StateTree/Source/StateTreeEditorModule/Private/AsyncStateTreeDiff.cpp
@@ -28,7 +28,7 @@ static bool AreObjectsEqual(const UObject* ObjectA, const UObject* ObjectB)
 		const FProperty* ClassProperty = ObjectA->GetClass()->PropertyLink;
 		while (ClassProperty)
 		{
-			if (!ClassProperty->Identical_InContainer(ObjectA, ObjectB, /*ArrayIndex*/ 0, PPF_DeepComparison))
+			if (!ClassProperty->Identical_InContainer(ObjectA, ObjectB, /*ArrayIndex*/ 0, PPF_DeepComparison | PPF_ForDiff))
 			{
 				return false;
 			}
@@ -42,8 +42,8 @@ static bool AreObjectsEqual(const UObject* ObjectA, const UObject* ObjectB)
 static bool AreNodesEqual(const FStateTreeEditorNode& NodeA, const FStateTreeEditorNode& NodeB)
 {
 	return AreObjectsEqual(NodeA.InstanceObject.Get(), NodeB.InstanceObject.Get())
-		&& NodeA.Node.Identical(&NodeB.Node, PPF_DeepComparison)
-		&& NodeA.Instance.Identical(&NodeB.Instance, PPF_DeepComparison)
+		&& NodeA.Node.Identical(&NodeB.Node, PPF_DeepComparison | PPF_ForDiff)
+		&& NodeA.Instance.Identical(&NodeB.Instance, PPF_DeepComparison | PPF_ForDiff)
 		&& NodeA.ExpressionIndent == NodeB.ExpressionIndent
 		&& NodeA.ExpressionOperand == NodeB.ExpressionOperand;
 }
@@ -64,21 +64,15 @@ static bool AreNodeArraysEqual(const TArray<FStateTreeEditorNode>& ArrayA, const
 	}
 	return true;
 }
-
-static bool AreStateTreeStateParametersEqual(const FStateTreeStateParameters& ParametersA, const FStateTreeStateParameters& ParametersB)
+static bool AreStateTreeStatePropertyBagsEqual(const FInstancedPropertyBag& ParametersA, const FInstancedPropertyBag& ParametersB)
 {
-	if (ParametersA.Parameters.GetNumPropertiesInBag() != ParametersB.Parameters.GetNumPropertiesInBag())
-	{
-		return false;
-	}
-	if (ParametersA.PropertyOverrides != ParametersB.PropertyOverrides)
+	if (ParametersA.GetNumPropertiesInBag() != ParametersB.GetNumPropertiesInBag())
 	{
 		return false;
 	}
-	const FInstancedPropertyBag ParamsA = ParametersA.Parameters;
-	const FInstancedPropertyBag ParamsB = ParametersB.Parameters;
-	const UPropertyBag* BagA = ParamsA.GetPropertyBagStruct();
-	const UPropertyBag* BagB = ParamsB.GetPropertyBagStruct();
+	
+	const UPropertyBag* BagA = ParametersA.GetPropertyBagStruct();
+	const UPropertyBag* BagB = ParametersB.GetPropertyBagStruct();
 	if (!BagA || !BagB)
 	{
 		return BagA == BagB;
@@ -96,8 +90,8 @@ static bool AreStateTreeStateParametersEqual(const FStateTreeStateParameters& Pa
 		}
 
 		const FName Name = DescsA[Index].Name;
-		TValueOrError<FString, EPropertyBagResult> SerializedA = ParamsA.GetValueSerializedString(Name);
-		TValueOrError<FString, EPropertyBagResult> SerializedB = ParamsB.GetValueSerializedString(Name);
+		TValueOrError<FString, EPropertyBagResult> SerializedA = ParametersA.GetValueSerializedString(Name);
+		TValueOrError<FString, EPropertyBagResult> SerializedB = ParametersB.GetValueSerializedString(Name);
 		if (SerializedA.HasError() || SerializedB.HasError())
 		{
 			return false;
@@ -116,6 +110,20 @@ static bool AreStateTreeStateParametersEqual(const FStateTreeStateParameters& Pa
 	return true;
 }
 
+static bool AreStateTreeStateParametersEqual(const FStateTreeStateParameters& ParametersA, const FStateTreeStateParameters& ParametersB)
+{
+	if (!AreStateTreeStatePropertyBagsEqual(ParametersA.Parameters, ParametersB.Parameters))
+	{
+		return false;
+	}
+	if (ParametersA.PropertyOverrides != ParametersB.PropertyOverrides)
+	{
+		return false;
+	}
+
+	return true;
+}
+
 static bool ArePropertiesEqual(const UStateTreeState* StateA, const UStateTreeState* StateB)
 {
 	if (!StateA || !StateB)
@@ -166,18 +174,7 @@ static bool AreTransitionsEqual(const UStateTreeState* StateA, const UStateTreeS
 			return TransitionA == TransitionB;
 		}
 		// Not checking transitions on IDs
-		const bool bEqual = TransitionA->Trigger == TransitionB->Trigger
-			&& TransitionA->RequiredEvent == TransitionB->RequiredEvent
-			&& TransitionA->State.Name == TransitionB->State.Name
-			&& TransitionA->State.ID == TransitionB->State.ID
-			&& TransitionA->State.LinkType == TransitionB->State.LinkType
-			&& TransitionA->State.StateHandle == TransitionB->State.StateHandle
-			&& TransitionA->Priority == TransitionB->Priority
-			&& TransitionA->bDelayTransition == TransitionB->bDelayTransition
-			&& TransitionA->DelayDuration == TransitionB->DelayDuration
-			&& TransitionA->DelayRandomVariance == TransitionB->DelayRandomVariance
-			&& AreNodeArraysEqual(TransitionA->Conditions, TransitionB->Conditions)
-			&& TransitionA->bTransitionEnabled == TransitionB->bTransitionEnabled;
+		const bool bEqual = FStateTreeTransition::StaticStruct()->CompareScriptStruct(TransitionB, TransitionB, 0);
 		if (!bEqual)
 		{
 			return false;
@@ -225,8 +222,8 @@ static bool AreStateTreePropertiesEqual(const UStateTreeEditorData* StateTreeDat
 		return false;
 	}
 
-	// Check the differences in Parameters
-	if (!AreStateTreeStateParametersEqual(StateTreeDataA->RootParameters, StateTreeDataB->RootParameters))
+	// Check the differences in StateTree root level parameters
+	if (!AreStateTreeStatePropertyBagsEqual(StateTreeDataA->GetRootParametersPropertyBag(), StateTreeDataB->GetRootParametersPropertyBag()))
 	{
 		return false;
 	}
diff --git a/Engine/Plugins/Runtime/StateTree/Source/StateTreeEditorModule/Private/Customizations/StateTreeEditorDataDetails.cpp b/Engine/Plugins/Runtime/StateTree/Source/StateTreeEditorModule/Private/Customizations/StateTreeEditorDataDetails.cpp
--- a/Engine/Plugins/Runtime/StateTree/Source/StateTreeEditorModule/Private/Customizations/StateTreeEditorDataDetails.cpp
+++ b/Engine/Plugins/Runtime/StateTree/Source/StateTreeEditorModule/Private/Customizations/StateTreeEditorDataDetails.cpp
@@ -40,6 +40,7 @@ void FStateTreeEditorDataDetails::CustomizeDetails(IDetailLayoutBuilder& DetailB
 		}
 	}
 	const UStateTreeSchema* Schema = EditorData ? EditorData->Schema : nullptr;
+	const FStateTreeDetailsCustomizationSettings* CustomizationSettings = Schema ? &Schema->GetDetailsCustomizationSettings() : nullptr;
 
 	// Common category
 	IDetailCategoryBuilder& CommonCategory = DetailBuilder.EditCategory(TEXT("Common"), LOCTEXT("EditorDataDetailsCommon", "Common"));
@@ -138,56 +139,62 @@ void FStateTreeEditorDataDetails::CustomizeDetails(IDetailLayoutBuilder& DetailB
 				];
 		}
 	}
-	
-	// Parameters category
-	const FText ParametersDisplayName = LOCTEXT("EditorDataDetailsParameters", "Parameters");
-	IDetailCategoryBuilder& ParametersCategory = DetailBuilder.EditCategory(TEXT("Parameters"), ParametersDisplayName);
-	ParametersCategory.SetSortOrder(2);
-	{
-		// Show parameters as a category.
-		TSharedRef<IPropertyUtilities> PropUtils = DetailBuilder.GetPropertyUtilities();
-		TSharedPtr<IPropertyHandle> RootParametersProperty = DetailBuilder.GetProperty(TEXT("RootParameters")); // FStateTreeStateParameters
-		check(RootParametersProperty);
-		RootParametersProperty->MarkHiddenByCustomization();
 
-		TSharedPtr<IPropertyHandle> PropertyBagParametersProperty = RootParametersProperty->GetChildHandle(TEXT("Parameters")); // FInstancedPropertyBag
-		check(PropertyBagParametersProperty);
+	if (CustomizationSettings && CustomizationSettings->bShowStateTreeParameters)
+	{
+		// Parameters category
+		const FText ParametersDisplayName = LOCTEXT("EditorDataDetailsParameters", "Parameters");
+		IDetailCategoryBuilder& ParametersCategory = DetailBuilder.EditCategory(TEXT("Parameters"), ParametersDisplayName);
+		ParametersCategory.SetSortOrder(2);
+		{
+			// Show parameters as a category.
+			TSharedRef<IPropertyUtilities> PropUtils = DetailBuilder.GetPropertyUtilities();
+			TSharedPtr<IPropertyHandle> PropertyBagParametersProperty = DetailBuilder.GetProperty(GET_MEMBER_NAME_CHECKED(UStateTreeEditorData, RootParameterPropertyBag)); // FInstancedPropertyBag
+			PropertyBagParametersProperty->MarkHiddenByCustomization();
+			check(PropertyBagParametersProperty);
 
-		const TSharedRef<SHorizontalBox> HeaderContentWidget = SNew(SHorizontalBox)
-			.IsEnabled(PropUtils, &IPropertyUtilities::IsPropertyEditingEnabled)
-			+ SHorizontalBox::Slot()
-			.VAlign(VAlign_Center)
-			.AutoWidth()
-			.Padding(FMargin(4.f, 0.f, 0.f, 0.f))
-			[
-				SNew(SImage)
-					.ColorAndOpacity(UE::StateTree::Colors::Blue)
-					.Image(FStateTreeEditorStyle::Get().GetBrush("StateTreeEditor.Parameters"))
-			]
-			+SHorizontalBox::Slot()
-			.VAlign(VAlign_Center)
-			.Padding(FMargin(4.f, 0.f, 0.f, 0.f))
-			.AutoWidth()
-			[
-				SNew(STextBlock)
-				.TextStyle(FStateTreeEditorStyle::Get(), "StateTree.Category")
-				.Text(ParametersDisplayName)
-			]
-			+SHorizontalBox::Slot()
-			.FillWidth(1.f)
-			.HAlign(HAlign_Right)
-			.VAlign(VAlign_Center)
-			[
-				FPropertyBagDetails::MakeAddPropertyWidget(PropertyBagParametersProperty, PropUtils, EPropertyBagPropertyType::Bool, FLinearColor(UE::StateTree::Colors::Blue)).ToSharedRef()
-			];
-		ParametersCategory.HeaderContent(HeaderContentWidget, /*FullRowContent*/true);
+			const TSharedRef<SHorizontalBox> HeaderContentWidget = SNew(SHorizontalBox)
+				.IsEnabled(PropUtils, &IPropertyUtilities::IsPropertyEditingEnabled)
+				+ SHorizontalBox::Slot()
+				.VAlign(VAlign_Center)
+				.AutoWidth()
+				.Padding(FMargin(4.f, 0.f, 0.f, 0.f))
+				[
+					SNew(SImage)
+						.ColorAndOpacity(UE::StateTree::Colors::Blue)
+						.Image(FStateTreeEditorStyle::Get().GetBrush("StateTreeEditor.Parameters"))
+				]
+				+SHorizontalBox::Slot()
+				.VAlign(VAlign_Center)
+				.Padding(FMargin(4.f, 0.f, 0.f, 0.f))
+				.AutoWidth()
+				[
+					SNew(STextBlock)
+					.TextStyle(FStateTreeEditorStyle::Get(), "StateTree.Category")
+					.Text(ParametersDisplayName)
+				]
+				+SHorizontalBox::Slot()
+				.FillWidth(1.f)
+				.HAlign(HAlign_Right)
+				.VAlign(VAlign_Center)
+				[
+					FPropertyBagDetails::MakeAddPropertyWidget(PropertyBagParametersProperty, PropUtils, EPropertyBagPropertyType::Bool, FLinearColor(UE::StateTree::Colors::Blue)).ToSharedRef()
+				];
+			ParametersCategory.HeaderContent(HeaderContentWidget, /*FullRowContent*/true);
 
-		TSharedRef<FPropertyBagInstanceDataDetails> InstanceDetails = MakeShareable(new FPropertyBagInstanceDataDetails(PropertyBagParametersProperty, PropUtils, false));
-		ParametersCategory.AddCustomBuilder(InstanceDetails);
+			TSharedRef<FPropertyBagInstanceDataDetails> InstanceDetails = MakeShareable(new FPropertyBagInstanceDataDetails(PropertyBagParametersProperty, PropUtils, false));
+			ParametersCategory.AddCustomBuilder(InstanceDetails);
+		}
 	}
+	else
+	{
+		TSharedRef<IPropertyUtilities> PropUtils = DetailBuilder.GetPropertyUtilities();
+		TSharedPtr<IPropertyHandle> PropertyBagParametersProperty = DetailBuilder.GetProperty(GET_MEMBER_NAME_CHECKED(UStateTreeEditorData, RootParameterPropertyBag)); // FInstancedPropertyBag
+		PropertyBagParametersProperty->MarkHiddenByCustomization();
+	}	
 
 	// Evaluators category
-	TSharedPtr<IPropertyHandle> EvaluatorsProperty = DetailBuilder.GetProperty(TEXT("Evaluators"));
+	TSharedPtr<IPropertyHandle> EvaluatorsProperty = DetailBuilder.GetProperty(GET_MEMBER_NAME_CHECKED(UStateTreeEditorData, Evaluators));
 	check(EvaluatorsProperty.IsValid());
 	const FName EvalCategoryName(TEXT("Evaluators"));
 	if (Schema && Schema->AllowEvaluators())
@@ -209,7 +216,7 @@ void FStateTreeEditorDataDetails::CustomizeDetails(IDetailLayoutBuilder& DetailB
 	}
 
 	// Global Tasks category
-	TSharedPtr<IPropertyHandle> GlobalTasksProperty = DetailBuilder.GetProperty(TEXT("GlobalTasks"));
+	TSharedPtr<IPropertyHandle> GlobalTasksProperty = DetailBuilder.GetProperty(GET_MEMBER_NAME_CHECKED(UStateTreeEditorData, GlobalTasks));
 	check(GlobalTasksProperty.IsValid());
 	const FName GlobalTasksCategoryName(TEXT("Global Tasks"));
 
@@ -225,7 +232,7 @@ void FStateTreeEditorDataDetails::CustomizeDetails(IDetailLayoutBuilder& DetailB
 		/*SortOrder*/4);
 
 	// Refresh the UI when the Schema changes.	
-	TSharedPtr<IPropertyHandle> SchemaProperty = DetailBuilder.GetProperty(TEXT("Schema"));
+	TSharedPtr<IPropertyHandle> SchemaProperty = DetailBuilder.GetProperty(GET_MEMBER_NAME_CHECKED(UStateTreeEditorData, Schema));
 	check(SchemaProperty.IsValid());
 	TSharedPtr<IPropertyUtilities> PropUtils = DetailBuilder.GetPropertyUtilities();
 	SchemaProperty->SetOnPropertyValueChanged(FSimpleDelegate::CreateLambda([PropUtils] ()
diff --git a/Engine/Plugins/Runtime/StateTree/Source/StateTreeEditorModule/Private/StateTreeCompiler.cpp b/Engine/Plugins/Runtime/StateTree/Source/StateTreeEditorModule/Private/StateTreeCompiler.cpp
--- a/Engine/Plugins/Runtime/StateTree/Source/StateTreeEditorModule/Private/StateTreeCompiler.cpp
+++ b/Engine/Plugins/Runtime/StateTree/Source/StateTreeEditorModule/Private/StateTreeCompiler.cpp
@@ -144,7 +144,7 @@ namespace UE::StateTree::Compiler
 		}
 	}
 
-	bool ValidateNoLevelActorReferences(FStateTreeCompilerLog& Log, const FStateTreeBindableStructDesc& NodeDesc, FStateTreeDataView NodeView, FStateTreeDataView InstanceView)
+	bool ValidateNoLevelActorReferences(FStateTreeCompilerLog& Log, const FStateTreeBindableStructDesc& NodeDesc, const FStateTreeDataView NodeView, const FStateTreeDataView InstanceView)
 	{
 		TSet<const UObject*> Visited;
 		TArray<const AActor*> LevelActors;
@@ -404,20 +404,22 @@ UStateTreeState* FStateTreeCompiler::GetState(const FGuid& StateID) const
 bool FStateTreeCompiler::CreateParameters()
 {
 	// Copy parameters from EditorData	
-	StateTree->Parameters = EditorData->RootParameters.Parameters;
+	StateTree->Parameters = EditorData->GetRootParametersPropertyBag();
 
 	// Mark parameters as binding source
+	const EStateTreeDataSourceType GlobalParameterDataType = UStateTreeSchema::CastToDataSourceType(EditorData->Schema->GetGlobalParameterDataType());
 	const FStateTreeBindableStructDesc ParametersDesc = {
 			UE::StateTree::Editor::GlobalStateName,
 			TEXT("Parameters"),
 			StateTree->Parameters.GetPropertyBagStruct(),
-			FStateTreeDataHandle(EStateTreeDataSourceType::GlobalParameterData),
+			FStateTreeDataHandle(GlobalParameterDataType),
 			EStateTreeBindableStructSource::Parameter,
-			EditorData->RootParameters.ID
+			EditorData->GetRootParametersGuid()
 	};
 	BindingsCompiler.AddSourceStruct(ParametersDesc);
 
-	if (!UE::StateTree::Compiler::ValidateNoLevelActorReferences(Log, ParametersDesc, FStateTreeDataView(), FStateTreeDataView(EditorData->RootParameters.Parameters.GetMutableValue())))
+	const FStateTreeDataView PropertyBagView(EditorData->GetRootParametersPropertyBag().GetPropertyBagStruct(), (uint8*)EditorData->GetRootParametersPropertyBag().GetValue().GetMemory());
+	if (!UE::StateTree::Compiler::ValidateNoLevelActorReferences(Log, ParametersDesc, FStateTreeDataView(), PropertyBagView))
 	{
 		return false;
 	}
diff --git a/Engine/Plugins/Runtime/StateTree/Source/StateTreeEditorModule/Private/StateTreeEditorData.cpp b/Engine/Plugins/Runtime/StateTree/Source/StateTreeEditorModule/Private/StateTreeEditorData.cpp
--- a/Engine/Plugins/Runtime/StateTree/Source/StateTreeEditorModule/Private/StateTreeEditorData.cpp
+++ b/Engine/Plugins/Runtime/StateTree/Source/StateTreeEditorModule/Private/StateTreeEditorData.cpp
@@ -15,6 +15,7 @@
 #include "Customizations/StateTreeBindingExtension.h"
 #include "Customizations/StateTreeEditorNodeUtils.h"
 #include "Modules/ModuleManager.h"
+#include "UObject/UE5SpecialProjectStreamObjectVersion.h"
 
 #if WITH_EDITOR
 #include "StructUtilsDelegates.h"
@@ -75,16 +76,26 @@ void UStateTreeEditorData::PostInitProperties()
 {
 	Super::PostInitProperties();
 
-	RootParameters.ID = FGuid::NewGuid();
+	RootParametersGuid = FGuid::NewGuid();
 
 #if WITH_EDITOR
-	OnObjectsReinstancedHandle = FCoreUObjectDelegates::OnObjectsReinstanced.AddUObject(this, &UStateTreeEditorData::OnObjectsReinstanced);
-	OnUserDefinedStructReinstancedHandle = UE::StructUtils::Delegates::OnUserDefinedStructReinstanced.AddUObject(this, &UStateTreeEditorData::OnUserDefinedStructReinstanced);
-	OnParametersChangedHandle = UE::StateTree::Delegates::OnParametersChanged.AddUObject(this, &UStateTreeEditorData::OnParametersChanged);
-	OnStateParametersChangedHandle = UE::StateTree::Delegates::OnStateParametersChanged.AddUObject(this, &UStateTreeEditorData::OnStateParametersChanged);
+	if (!HasAnyFlags(RF_ClassDefaultObject))
+	{
+		OnObjectsReinstancedHandle = FCoreUObjectDelegates::OnObjectsReinstanced.AddUObject(this, &UStateTreeEditorData::OnObjectsReinstanced);
+		OnUserDefinedStructReinstancedHandle = UE::StructUtils::Delegates::OnUserDefinedStructReinstanced.AddUObject(this, &UStateTreeEditorData::OnUserDefinedStructReinstanced);
+		OnParametersChangedHandle = UE::StateTree::Delegates::OnParametersChanged.AddUObject(this, &UStateTreeEditorData::OnParametersChanged);
+		OnStateParametersChangedHandle = UE::StateTree::Delegates::OnStateParametersChanged.AddUObject(this, &UStateTreeEditorData::OnStateParametersChanged);
+	}
+	
 #endif
 }
 
+void UStateTreeEditorData::Serialize(FArchive& Ar)
+{
+	Super::Serialize(Ar);	
+	Ar.UsingCustomVersion(FUE5SpecialProjectStreamObjectVersion::GUID);
+}
+
 #if WITH_EDITOR
 
 void UStateTreeEditorData::BeginDestroy()
@@ -204,11 +215,20 @@ void UStateTreeEditorData::PostLoad()
 {
 	Super::PostLoad();
 
+	if (GetLinkerCustomVersion(FUE5SpecialProjectStreamObjectVersion::GUID) < FUE5SpecialProjectStreamObjectVersion::StateTreeGlobalParameterChanges)
+	{
+		PRAGMA_DISABLE_DEPRECATION_WARNINGS
+		RootParameterPropertyBag = RootParameters.Parameters;
+		RootParametersGuid = RootParameters.ID;
+		PRAGMA_ENABLE_DEPRECATION_WARNINGS
+	}
+
 	// Ensure the schema and states have had their PostLoad() fixed applied as we may need them in the later calls (or StateTree compile which might be calling this).
 	if (Schema)
 	{
 		Schema->ConditionalPostLoad();
 	}
+
 	VisitHierarchy([](UStateTreeState& State, UStateTreeState* ParentState) mutable 
 	{
 		State.ConditionalPostLoad();
@@ -243,7 +263,7 @@ void UStateTreeEditorData::PostEditChangeChainProperty(FPropertyChangedChainEven
 		{			
 			UE::StateTree::Delegates::OnSchemaChanged.Broadcast(*StateTree);
 		}
-		else if (MemberName == GET_MEMBER_NAME_CHECKED(UStateTreeEditorData, RootParameters))
+		else if (MemberName == GET_MEMBER_NAME_CHECKED(UStateTreeEditorData, RootParameterPropertyBag))
 		{
 			UE::StateTree::Delegates::OnParametersChanged.Broadcast(*StateTree);
 		}
@@ -570,7 +590,7 @@ void UStateTreeEditorData::AppendBindablePropertyFunctionStructs(TArray<TInstanc
 
 bool UStateTreeEditorData::CanCreateParameter(const FGuid StructID) const
 {
-	if (RootParameters.ID == StructID)
+	if (RootParametersGuid == StructID)
 	{
 		return true;
 	}
@@ -600,9 +620,9 @@ void UStateTreeEditorData::CreateParametersForStruct(const FGuid StructID, TArra
 	const UStateTree* StateTree = GetTypedOuter<UStateTree>();
 	checkf(StateTree, TEXT("UStateTreeEditorData should only be allocated within a UStateTree"));
 
-	if (RootParameters.ID == StructID)
+	if (RootParametersGuid == StructID)
 	{
-		UE::PropertyBinding::CreateUniquelyNamedPropertiesInPropertyBag(InOutCreationDescs, RootParameters.Parameters);
+		CreateRootProperties(InOutCreationDescs);
 		UE::StateTree::Delegates::OnParametersChanged.Broadcast(*StateTree);
 		return;
 	}
@@ -1010,7 +1030,7 @@ void UStateTreeEditorData::DuplicateIDs()
 		};
 
 		// Do not use the VisitGlobalNodes because the schema should not be included in OldToNewIDs 
-		OldToNewIDs.Add(RootParameters.ID, FGuid::NewGuid());
+		OldToNewIDs.Add(RootParametersGuid, FGuid::NewGuid());
 		AddIds(Evaluators);
 		AddIds(GlobalTasks);
 		for (FStateTreeEditorColor& Color : Colors)
@@ -1473,12 +1493,12 @@ EStateTreeVisitor UStateTreeEditorData::VisitGlobalNodes(TFunctionRef<EStateTree
 	{
 		FStateTreeBindableStructDesc Desc;
 		Desc.StatePath = UE::StateTree::Editor::GlobalStateName;
-		Desc.Struct = RootParameters.Parameters.GetPropertyBagStruct();
+		Desc.Struct = GetRootParametersPropertyBag().GetPropertyBagStruct();
 		Desc.Name = FName(TEXT("Parameters"));
-		Desc.ID = RootParameters.ID;
+		Desc.ID = RootParametersGuid;
 		Desc.DataSource = EStateTreeBindableStructSource::Parameter;
 		
-		if (InFunc(nullptr, Desc, FStateTreeDataView(const_cast<FInstancedPropertyBag&>(RootParameters.Parameters).GetMutableValue())) == EStateTreeVisitor::Break)
+		if (InFunc(nullptr, Desc, FStateTreeDataView(const_cast<FInstancedPropertyBag&>(GetRootParametersPropertyBag()).GetMutableValue())) == EStateTreeVisitor::Break)
 		{
 			return EStateTreeVisitor::Break;
 		}
diff --git a/Engine/Plugins/Runtime/StateTree/Source/StateTreeEditorModule/Private/StateTreePropertyBindingCompiler.cpp b/Engine/Plugins/Runtime/StateTree/Source/StateTreeEditorModule/Private/StateTreePropertyBindingCompiler.cpp
--- a/Engine/Plugins/Runtime/StateTree/Source/StateTreeEditorModule/Private/StateTreePropertyBindingCompiler.cpp
+++ b/Engine/Plugins/Runtime/StateTree/Source/StateTreeEditorModule/Private/StateTreePropertyBindingCompiler.cpp
@@ -368,8 +368,9 @@ bool FStateTreePropertyBindingCompiler::CompileReferences(const FStateTreeBindab
 		if (!ReferenceIndex.IsValid())
 		{
 			// If referencing another non global or subtree parameter PropertyRef, reuse it's index.
-			if (UE::StateTree::PropertyRefHelpers::IsPropertyRef(*SourceIndirections.Last().GetProperty()) 
-				&& SourceStruct->DataHandle.GetSource() != EStateTreeDataSourceType::GlobalParameterData 
+			if (UE::StateTree::PropertyRefHelpers::IsPropertyRef(*SourceIndirections.Last().GetProperty())
+				&& SourceStruct->DataHandle.GetSource() != EStateTreeDataSourceType::GlobalParameterData
+				&& SourceStruct->DataHandle.GetSource() != EStateTreeDataSourceType::ExternalGlobalParameterData
 				&& SourceStruct->DataHandle.GetSource() != EStateTreeDataSourceType::SubtreeParameterData)
 			{
 				const FCompiledReference* ReferencedReference = CompiledReferences.FindByPredicate([&Binding](const FCompiledReference& CompiledReference)
diff --git a/Engine/Plugins/Runtime/StateTree/Source/StateTreeEditorModule/Public/StateTreeEditorData.h b/Engine/Plugins/Runtime/StateTree/Source/StateTreeEditorModule/Public/StateTreeEditorData.h
--- a/Engine/Plugins/Runtime/StateTree/Source/StateTreeEditorModule/Public/StateTreeEditorData.h
+++ b/Engine/Plugins/Runtime/StateTree/Source/StateTreeEditorModule/Public/StateTreeEditorData.h
@@ -62,6 +62,7 @@ public:
 	UStateTreeEditorData();
 
 	virtual void PostInitProperties() override;
+	virtual void Serialize(FArchive& Ar) override;
 	
 	//~ Begin IStateTreeEditorPropertyBindingsOwner interface
 	virtual void GetBindableStructs(const FGuid TargetStructID, TArray<TInstancedStruct<FPropertyBindingBindableStructDescriptor>>& OutStructDescs) const override;
@@ -114,6 +115,18 @@ public:
 	virtual void PostDuplicate(EDuplicateMode::Type DuplicateMode) override;
 #endif
 
+	/** @return the public parameters ID that could be used for bindings within the Tree. */
+	FGuid GetRootParametersGuid() const
+	{
+		return RootParametersGuid;
+	}
+
+	/** @return the public parameters that could be used for bindings within the Tree. */
+	virtual const FInstancedPropertyBag& GetRootParametersPropertyBag() const
+	{
+		return RootParameterPropertyBag;
+	}
+
 	/** @returns parent state of a struct, or nullptr if not found. */
 	const UStateTreeState* GetStateByStructID(const FGuid TargetStructID) const;
 
@@ -306,6 +319,8 @@ PRAGMA_ENABLE_DEPRECATION_WARNINGS
 		return Colors.Find(FStateTreeEditorColor(ColorRef));
 	}
 
+	virtual void CreateRootProperties(TArrayView<UE::PropertyBinding::FPropertyCreationDescriptor> InOutCreationDescs) { UE::PropertyBinding::CreateUniquelyNamedPropertiesInPropertyBag(InOutCreationDescs, RootParameterPropertyBag); }
+
 private:
 	virtual EStateTreeVisitor EnumerateBindablePropertyFunctionNodes(TFunctionRef<EStateTreeVisitor(const UScriptStruct* NodeStruct, const FStateTreeBindableStructDesc& Desc, const FStateTreeDataView Value)> InFunc) const override;
 
@@ -329,9 +344,20 @@ public:
 	TObjectPtr<UStateTreeSchema> Schema = nullptr;
 
 	/** Public parameters that could be used for bindings within the Tree. */
-	UPROPERTY(EditDefaultsOnly, Category = Parameters)
+	UE_DEPRECATED(5.6, "Public access to RootParameters is deprecated. Use GetRootParametersPropertyBag")
+	UPROPERTY(meta = (DeprecatedProperty))
 	FStateTreeStateParameters RootParameters;
 
+private:
+	/** Public parameters ID that could be used for bindings within the Tree. */
+	UPROPERTY()
+	FGuid RootParametersGuid;
+
+	/** Public parameters property bag that could be used for bindings within the Tree. */
+	UPROPERTY(EditDefaultsOnly, Category = Parameters)
+	FInstancedPropertyBag RootParameterPropertyBag;
+
+public:
 	UPROPERTY(EditDefaultsOnly, Category = "Evaluators", meta = (BaseStruct = "/Script/StateTreeModule.StateTreeEvaluatorBase", BaseClass = "/Script/StateTreeModule.StateTreeEvaluatorBlueprintBase"))
 	TArray<FStateTreeEditorNode> Evaluators;
 
@@ -364,4 +390,6 @@ public:
 	 */
 	UPROPERTY(DuplicateTransient)
 	TArray<FStateTreeEditorDelegateDispatcherCompiledBinding> CompiledDispatchers;
+
+	friend class FStateTreeEditorDataDetails;
 };
diff --git a/Engine/Plugins/Runtime/StateTree/Source/StateTreeModule/Private/StateTree.cpp b/Engine/Plugins/Runtime/StateTree/Source/StateTreeModule/Private/StateTree.cpp
--- a/Engine/Plugins/Runtime/StateTree/Source/StateTreeModule/Private/StateTree.cpp
+++ b/Engine/Plugins/Runtime/StateTree/Source/StateTreeModule/Private/StateTree.cpp
@@ -816,7 +816,8 @@ bool UStateTree::PatchBindings()
 	}
 
 	// Update property bag structs before resolving binding.
-	if (FStateTreeBindableStructDesc* RootParamsDesc = GetSourceStructByHandle(FStateTreeDataHandle(EStateTreeDataSourceType::GlobalParameterData)))
+	const EStateTreeDataSourceType GlobalParameterDataType = UStateTreeSchema::CastToDataSourceType(Schema->GetGlobalParameterDataType());
+	if (FStateTreeBindableStructDesc* RootParamsDesc = GetSourceStructByHandle(FStateTreeDataHandle(GlobalParameterDataType)))
 	{
 		RootParamsDesc->Struct = Parameters.GetPropertyBagStruct();
 	}
@@ -899,7 +900,7 @@ bool UStateTree::PatchBindings()
 	TMap<FStateTreeIndex16, FStateTreeDataView> BindingBatchDataView;
 
 	// Tree parameters
-	DataViews.Add(FStateTreeDataHandle(EStateTreeDataSourceType::GlobalParameterData), Parameters.GetMutableValue());
+	DataViews.Add(FStateTreeDataHandle(GlobalParameterDataType), Parameters.GetMutableValue());
 
 	// Setup data views for context data. Since the external data is passed at runtime, we can only provide the type.
 	for (const FStateTreeExternalDataDesc& DataDesc : ContextDataDescs)
diff --git a/Engine/Plugins/Runtime/StateTree/Source/StateTreeModule/Private/StateTreeExecutionContext.cpp b/Engine/Plugins/Runtime/StateTree/Source/StateTreeModule/Private/StateTreeExecutionContext.cpp
--- a/Engine/Plugins/Runtime/StateTree/Source/StateTreeModule/Private/StateTreeExecutionContext.cpp
+++ b/Engine/Plugins/Runtime/StateTree/Source/StateTreeModule/Private/StateTreeExecutionContext.cpp
@@ -864,6 +864,36 @@ const FStateTreeReference* FStateTreeExecutionContext::GetLinkedStateTreeOverrid
 	return nullptr;
 }
 
+bool FStateTreeExecutionContext::FExternalGlobalParameters::Add(const FPropertyBindingCopyInfo& Copy, uint8* InParameterMemory)
+{
+	const int32 TypeHash = HashCombine(GetTypeHash(Copy.SourceLeafProperty), GetTypeHash(Copy.SourceIndirection));
+	const int32 NumMappings = Mappings.Num();
+	Mappings.Add(TypeHash, InParameterMemory);
+	return Mappings.Num() > NumMappings;			
+}
+
+uint8* FStateTreeExecutionContext::FExternalGlobalParameters::Find(const FPropertyBindingCopyInfo& Copy) const
+{
+	const int32 TypeHash = HashCombine(GetTypeHash(Copy.SourceLeafProperty), GetTypeHash(Copy.SourceIndirection));
+	if(uint8* const* MappingPtr = Mappings.Find(TypeHash))
+	{
+		return *MappingPtr;
+	}
+		
+	checkf(false, TEXT("Missing external parameter data"));
+	return nullptr;
+}
+
+void FStateTreeExecutionContext::FExternalGlobalParameters::Reset()
+{
+	Mappings.Reset();
+}
+
+void FStateTreeExecutionContext::SetExternalGlobalParameters(const FExternalGlobalParameters* Parameters)
+{
+	ExternalGlobalParameters = Parameters;
+}
+
 bool FStateTreeExecutionContext::AreContextDataViewsValid() const
 {
 	if (!IsValid())
@@ -2177,6 +2207,17 @@ FStateTreeDataView FStateTreeExecutionContext::GetDataView(const FStateTreeExecu
 	}
 }
 
+FStateTreeDataView FStateTreeExecutionContext::GetDataView(const FStateTreeExecutionFrame* ParentFrame, const FStateTreeExecutionFrame& CurrentFrame, const FPropertyBindingCopyInfo& CopyInfo)
+{	
+	const FStateTreeDataHandle Handle = CopyInfo.SourceDataHandle.Get<FStateTreeDataHandle>();
+	if (Handle.GetSource() == EStateTreeDataSourceType::ExternalGlobalParameterData)
+	{
+		return GetDataViewOrTemporary(ParentFrame, CurrentFrame, CopyInfo);
+	}
+
+	return GetDataView(ParentFrame, CurrentFrame, Handle);
+}
+
 EStateTreeRunStatus FStateTreeExecutionContext::ForceTransition(const FRecordedStateTreeTransitionResult& Transition)
 {
 	if (!IsValid())
@@ -2225,6 +2266,12 @@ FStateTreeDataView FStateTreeExecutionContext::GetDataViewFromInstanceStorage(FS
 		}
 		return InstanceDataStorage.GetMutableGlobalParameters();
 
+	case EStateTreeDataSourceType::ExternalGlobalParameterData:
+	    {		    
+			checkf(false, TEXT("External global parameter data currently not supported for linked state-trees"));
+	    }
+	    break;
+
 	case EStateTreeDataSourceType::SubtreeParameterData:
 		{
 			// Defined in parent frame.
@@ -2333,6 +2380,11 @@ bool FStateTreeExecutionContext::IsHandleSourceValid(const FStateTreeExecutionFr
 			? IsHandleSourceValid(nullptr, *ParentFrame, CurrentFrame.GlobalParameterDataHandle)
 			: CurrentFrame.GlobalParameterDataHandle.IsValid();
 
+	case EStateTreeDataSourceType::ExternalGlobalParameterData:
+		{
+			checkf(false, TEXT("External global parameter data currently not supported for linked state-trees"));
+		}
+		break;
 	case EStateTreeDataSourceType::SubtreeParameterData:
 		if (ParentFrame)
 		{
@@ -2370,6 +2422,17 @@ bool FStateTreeExecutionContext::IsHandleSourceValid(const FStateTreeExecutionFr
 	return false;
 }
 
+bool FStateTreeExecutionContext::IsHandleSourceValid(const FStateTreeExecutionFrame* ParentFrame, const FStateTreeExecutionFrame& CurrentFrame, const FPropertyBindingCopyInfo& CopyInfo) const
+{
+	const FStateTreeDataHandle Handle = CopyInfo.SourceDataHandle.Get<FStateTreeDataHandle>();
+	if (Handle.GetSource() == EStateTreeDataSourceType::ExternalGlobalParameterData)
+	{
+		return ExternalGlobalParameters ? ExternalGlobalParameters->Find(CopyInfo) != nullptr : false;
+	}
+
+	return IsHandleSourceValid(ParentFrame, CurrentFrame, Handle);
+}
+
 FStateTreeDataView FStateTreeExecutionContext::GetDataViewOrTemporary(const FStateTreeExecutionFrame* ParentFrame, const FStateTreeExecutionFrame& CurrentFrame, const FStateTreeDataHandle Handle)
 {
 	if (IsHandleSourceValid(ParentFrame, CurrentFrame, Handle))
@@ -2397,6 +2460,12 @@ FStateTreeDataView FStateTreeExecutionContext::GetDataViewOrTemporary(const FSta
 		}
 		break;
 		
+	case EStateTreeDataSourceType::ExternalGlobalParameterData:
+	    {
+			checkf(false, TEXT("External global parameter data currently not supported for linked state-trees"));
+	    }
+		break;
+		
 	case EStateTreeDataSourceType::SubtreeParameterData:
 		if (ParentFrame)
 		{
@@ -2427,6 +2496,18 @@ FStateTreeDataView FStateTreeExecutionContext::GetDataViewOrTemporary(const FSta
 	return {};
 }
 
+FStateTreeDataView FStateTreeExecutionContext::GetDataViewOrTemporary(const FStateTreeExecutionFrame* ParentFrame, const FStateTreeExecutionFrame& CurrentFrame, const FPropertyBindingCopyInfo& CopyInfo)
+{
+	const FStateTreeDataHandle Handle = CopyInfo.SourceDataHandle.Get<FStateTreeDataHandle>();
+	if (Handle.GetSource() == EStateTreeDataSourceType::ExternalGlobalParameterData)
+	{
+		uint8* MemoryPtr = ExternalGlobalParameters->Find(CopyInfo);
+		return FStateTreeDataView(CopyInfo.SourceStructType, MemoryPtr);	
+	}
+
+	return GetDataViewOrTemporary(ParentFrame, CurrentFrame, Handle);	
+}
+
 FStateTreeDataView FStateTreeExecutionContext::AddTemporaryInstance(const FStateTreeExecutionFrame& Frame, const FStateTreeIndex16 OwnerNodeIndex, const FStateTreeDataHandle DataHandle, FConstStructView NewInstanceData)
 {
 	const FStructView NewInstance = Storage.AddTemporaryInstance(Owner, Frame, OwnerNodeIndex, DataHandle, NewInstanceData);
@@ -2451,7 +2532,7 @@ bool FStateTreeExecutionContext::CopyBatchOnActiveInstances(const FStateTreeExec
 	bool bSucceed = true;
 	for (const FPropertyBindingCopyInfo& Copy : CurrentFrame.StateTree->PropertyBindings.Super::GetBatchCopies(Batch))
 	{
-		const FStateTreeDataView SourceView = GetDataView(ParentFrame, CurrentFrame, Copy.SourceDataHandle.Get<FStateTreeDataHandle>());
+		const FStateTreeDataView SourceView = GetDataView(ParentFrame, CurrentFrame, Copy);
 		bSucceed &= CurrentFrame.StateTree->PropertyBindings.Super::CopyProperty(Copy, SourceView, TargetView);
 	}
 	return bSucceed;
@@ -2471,7 +2552,7 @@ bool FStateTreeExecutionContext::CopyBatchWithValidation(const FStateTreeExecuti
 	bool bSucceed = true;
 	for (const FPropertyBindingCopyInfo& Copy : CurrentFrame.StateTree->PropertyBindings.Super::GetBatchCopies(Batch))
 	{
-		const FStateTreeDataView SourceView = GetDataViewOrTemporary(ParentFrame, CurrentFrame, Copy.SourceDataHandle.Get<FStateTreeDataHandle>());
+		const FStateTreeDataView SourceView = GetDataViewOrTemporary(ParentFrame, CurrentFrame, Copy);
 		if (!SourceView.IsValid())
 		{
 			bSucceed = false;
diff --git a/Engine/Plugins/Runtime/StateTree/Source/StateTreeModule/Private/StateTreeSchema.cpp b/Engine/Plugins/Runtime/StateTree/Source/StateTreeModule/Private/StateTreeSchema.cpp
--- a/Engine/Plugins/Runtime/StateTree/Source/StateTreeModule/Private/StateTreeSchema.cpp
+++ b/Engine/Plugins/Runtime/StateTree/Source/StateTreeModule/Private/StateTreeSchema.cpp
@@ -1,6 +1,7 @@
 // Copyright Epic Games, Inc. All Rights Reserved.
 
 #include "StateTreeSchema.h"
+#include "StateTreeTypes.h"
 #include "Blueprint/StateTreeNodeBlueprintBase.h"
 
 #include UE_INLINE_GENERATED_CPP_BY_NAME(StateTreeSchema)
@@ -10,3 +11,7 @@ bool UStateTreeSchema::IsChildOfBlueprintBase(const UClass* InClass) const
 	return InClass->IsChildOf<UStateTreeNodeBlueprintBase>();
 }
 
+EStateTreeDataSourceType UStateTreeSchema::CastToDataSourceType(EGlobalParameterDataType Value)
+{
+	return Value == EGlobalParameterDataType::ExternalGlobalParameterData ? EStateTreeDataSourceType::ExternalGlobalParameterData : EStateTreeDataSourceType::GlobalParameterData;
+}
diff --git a/Engine/Plugins/Runtime/StateTree/Source/StateTreeModule/Public/StateTreeExecutionContext.h b/Engine/Plugins/Runtime/StateTree/Source/StateTreeModule/Public/StateTreeExecutionContext.h
--- a/Engine/Plugins/Runtime/StateTree/Source/StateTreeModule/Public/StateTreeExecutionContext.h
+++ b/Engine/Plugins/Runtime/StateTree/Source/StateTreeModule/Public/StateTreeExecutionContext.h
@@ -288,6 +288,18 @@ public:
 	/** @return the first state tree reference set by SetLinkedStateTreeOverrides that matches the StateTag. Or null if not found. */
 	const FStateTreeReference* GetLinkedStateTreeOverrideForTag(const FGameplayTag StateTag) const;
 
+	/** Structure to-be-populated and set for any StateTree using any EStateTreeDataSourceType::ExternalGlobalParameterData bindings */
+	struct STATETREEMODULE_API FExternalGlobalParameters
+	{
+		/* Add memory mapping, this expects InParameterMemory to resolve correctly for the SourceLeafProperty and SourceIndirection */
+		bool Add(const FPropertyBindingCopyInfo& Copy, uint8* InParameterMemory);
+		uint8* Find(const FPropertyBindingCopyInfo& Copy) const;
+		void Reset();
+	private:
+		TMap<uint32, uint8*> Mappings;
+	};
+	void SetExternalGlobalParameters(const FExternalGlobalParameters* Parameters);
+
 	/** @return const references to the instance data in use, or nullptr if the context is not valid. */
 	const FStateTreeInstanceData* GetInstanceData() const
 	{
@@ -953,12 +965,15 @@ protected:
 
 	/** @return data view of the specified handle relative to given frame. */
 	FStateTreeDataView GetDataView(const FStateTreeExecutionFrame* ParentFrame, const FStateTreeExecutionFrame& CurrentFrame, const FStateTreeDataHandle Handle);
+	FStateTreeDataView GetDataView(const FStateTreeExecutionFrame* ParentFrame, const FStateTreeExecutionFrame& CurrentFrame, const FPropertyBindingCopyInfo& CopyInfo);
 
 	/** @return true if handle source is valid cified handle relative to given frame. */
 	bool IsHandleSourceValid(const FStateTreeExecutionFrame* ParentFrame, const FStateTreeExecutionFrame& CurrentFrame, const FStateTreeDataHandle Handle) const;
+	bool IsHandleSourceValid(const FStateTreeExecutionFrame* ParentFrame, const FStateTreeExecutionFrame& CurrentFrame, const FPropertyBindingCopyInfo& CopyInfo) const;
 
 	/** @return data view of the specified handle relative to the given frame, or tries to find a matching temporary instance. */
 	FStateTreeDataView GetDataViewOrTemporary(const FStateTreeExecutionFrame* ParentFrame, const FStateTreeExecutionFrame& CurrentFrame, const FStateTreeDataHandle Handle);
+	FStateTreeDataView GetDataViewOrTemporary(const FStateTreeExecutionFrame* ParentFrame, const FStateTreeExecutionFrame& CurrentFrame, const FPropertyBindingCopyInfo& CopyInfo);
 
 	/**
 	 * Adds a temporary instance that can be located using frame and data handle later.
@@ -1207,6 +1222,9 @@ protected:
 
 	/** Captured snapshots for transition results that can be used to recreate transitions. This array is only populated if bRecordTransitions is true. */
 	TArray<FRecordedStateTreeTransitionResult> RecordedTransitions;
+
+	/** Memory mapping structure used for redirecting property-bag copies to external (raw) memory pointers */
+	const FExternalGlobalParameters* ExternalGlobalParameters = nullptr;
 };
 
 /**
diff --git a/Engine/Plugins/Runtime/StateTree/Source/StateTreeModule/Public/StateTreePropertyRef.h b/Engine/Plugins/Runtime/StateTree/Source/StateTreeModule/Public/StateTreePropertyRef.h
--- a/Engine/Plugins/Runtime/StateTree/Source/StateTreeModule/Public/StateTreePropertyRef.h
+++ b/Engine/Plugins/Runtime/StateTree/Source/StateTreeModule/Public/StateTreePropertyRef.h
@@ -32,7 +32,9 @@ namespace UE::StateTree::PropertyRefHelpers
 			// If that's the case, referenced PropertyRef is obtained and we recursively take the address where it points to.
 			if (UE::StateTree::PropertyRefHelpers::IsPropertyRef(*PropertyAccess->SourceLeafProperty))
 			{
-				check(PropertyAccess->SourceDataHandle.GetSource() == EStateTreeDataSourceType::GlobalParameterData || PropertyAccess->SourceDataHandle.GetSource() == EStateTreeDataSourceType::SubtreeParameterData);
+				check(PropertyAccess->SourceDataHandle.GetSource() == EStateTreeDataSourceType::GlobalParameterData ||
+					PropertyAccess->SourceDataHandle.GetSource() == EStateTreeDataSourceType::ExternalGlobalParameterData ||
+					PropertyAccess->SourceDataHandle.GetSource() == EStateTreeDataSourceType::SubtreeParameterData);
 
 				if (ParentExecutionFrame == nullptr)
 				{
diff --git a/Engine/Plugins/Runtime/StateTree/Source/StateTreeModule/Public/StateTreeSchema.h b/Engine/Plugins/Runtime/StateTree/Source/StateTreeModule/Public/StateTreeSchema.h
--- a/Engine/Plugins/Runtime/StateTree/Source/StateTreeModule/Public/StateTreeSchema.h
+++ b/Engine/Plugins/Runtime/StateTree/Source/StateTreeModule/Public/StateTreeSchema.h
@@ -7,6 +7,15 @@
 
 struct FStateTreeExternalDataDesc;
 struct FStateTreeLinker;
+enum class EStateTreeDataSourceType : uint8;
+
+
+#if WITH_EDITOR
+struct FStateTreeDetailsCustomizationSettings
+{
+	bool bShowStateTreeParameters : 1 = true;
+};
+#endif // WITH_EDITOR
 
 /**
  * Schema describing which inputs, evaluators, and tasks a StateTree can contain.
@@ -61,6 +70,19 @@ public:
 		return {};
 	}
 
+	enum class EGlobalParameterDataType : uint8
+	{
+		GlobalParameterData,
+		ExternalGlobalParameterData,
+	};
+	static EStateTreeDataSourceType CastToDataSourceType(EGlobalParameterDataType Value);
+
+	/** @return the global parameter type used by the schema. */
+	virtual EGlobalParameterDataType GetGlobalParameterDataType() const
+	{
+		return EGlobalParameterDataType::GlobalParameterData;
+	}
+
 	/** Resolves schema references to other StateTree data. */
 	virtual bool Link(FStateTreeLinker& Linker)
 	{
@@ -91,5 +113,12 @@ public:
 	{
 		return true;
 	}
+
+	virtual const FStateTreeDetailsCustomizationSettings& GetDetailsCustomizationSettings() const
+	{
+		static FStateTreeDetailsCustomizationSettings Settings;
+		return Settings;
+	}
+	
 #endif // WITH_EDITOR
 };
diff --git a/Engine/Plugins/Runtime/StateTree/Source/StateTreeModule/Public/StateTreeTypes.h b/Engine/Plugins/Runtime/StateTree/Source/StateTreeModule/Public/StateTreeTypes.h
--- a/Engine/Plugins/Runtime/StateTree/Source/StateTreeModule/Public/StateTreeTypes.h
+++ b/Engine/Plugins/Runtime/StateTree/Source/StateTreeModule/Public/StateTreeTypes.h
@@ -320,6 +320,9 @@ enum class EStateTreeDataSourceType : uint8
 
 	/** Event used in state selection. */
 	StateEvent,
+
+	/** Global parameters provided externally */
+	ExternalGlobalParameterData,
 };
 
 /** Handle to a StateTree data */
@@ -355,7 +358,7 @@ struct STATETREEMODULE_API FStateTreeDataHandle
 		// Require valid state for active instance data
 		check(Source != EStateTreeDataSourceType::ActiveInstanceData || (Source == EStateTreeDataSourceType::ActiveInstanceData && StateHandle.IsValid()));
 		check(Source != EStateTreeDataSourceType::ActiveInstanceDataObject || (Source == EStateTreeDataSourceType::ActiveInstanceDataObject && StateHandle.IsValid()));
-		check(Source == EStateTreeDataSourceType::GlobalParameterData || IsValidIndex(InIndex));
+		check(Source == EStateTreeDataSourceType::GlobalParameterData || Source == EStateTreeDataSourceType::ExternalGlobalParameterData || IsValidIndex(InIndex));
 	}
 
 	explicit FStateTreeDataHandle(const EStateTreeDataSourceType InSource, const int32 InIndex, const FStateTreeStateHandle InStateHandle = FStateTreeStateHandle::Invalid)
@@ -365,7 +368,7 @@ struct STATETREEMODULE_API FStateTreeDataHandle
 		// Require valid state for active instance data
 		check(Source != EStateTreeDataSourceType::ActiveInstanceData || (Source == EStateTreeDataSourceType::ActiveInstanceData && StateHandle.IsValid()));
 		check(Source != EStateTreeDataSourceType::ActiveInstanceDataObject || (Source == EStateTreeDataSourceType::ActiveInstanceDataObject && StateHandle.IsValid()));
-		check(Source == EStateTreeDataSourceType::GlobalParameterData || IsValidIndex(InIndex));
+		check(Source == EStateTreeDataSourceType::GlobalParameterData || Source == EStateTreeDataSourceType::ExternalGlobalParameterData || IsValidIndex(InIndex));
 		Index = static_cast<uint16>(InIndex);
 	}
 
@@ -455,6 +458,8 @@ struct STATETREEMODULE_API FStateTreeDataHandle
 			return FString::Printf(TEXT("Context[%d]"), Index);
 		case EStateTreeDataSourceType::GlobalParameterData:
 			return FString::Printf(TEXT("GlobalParam"));
+		case EStateTreeDataSourceType::ExternalGlobalParameterData:
+			return FString::Printf(TEXT("ExternalGlobalParam"));
 		case EStateTreeDataSourceType::SubtreeParameterData:
 			return FString::Printf(TEXT("SubtreeParam[%d]"), Index);
 		case EStateTreeDataSourceType::StateParameterData:
diff --git a/Engine/Plugins/Runtime/StateTree/Source/StateTreeTestSuite/Private/StateTreeBindingTest.cpp b/Engine/Plugins/Runtime/StateTree/Source/StateTreeTestSuite/Private/StateTreeBindingTest.cpp
--- a/Engine/Plugins/Runtime/StateTree/Source/StateTreeTestSuite/Private/StateTreeBindingTest.cpp
+++ b/Engine/Plugins/Runtime/StateTree/Source/StateTreeTestSuite/Private/StateTreeBindingTest.cpp
@@ -531,11 +531,12 @@ struct FStateTreeTest_MutableArray : FStateTreeTestBase
 			UStateTreeEditorData& EditorData = *Cast<UStateTreeEditorData>(StateTree.EditorData);
 			{
 				// Parameters
-				EditorData.RootParameters.Parameters.AddProperty("Value", EPropertyBagPropertyType::Int32);
-				EditorData.RootParameters.Parameters.SetValueInt32("Value", -111);
-				EditorData.RootParameters.Parameters.AddContainerProperty("ArrayValue", FPropertyBagContainerTypes(EPropertyBagContainerType::Array), EPropertyBagPropertyType::Int32, nullptr);
-				EditorData.RootParameters.Parameters.AddProperty("ArrayValue", EPropertyBagPropertyType::Int32);
-				FPropertyBagArrayRef ValueArrayRef = EditorData.RootParameters.Parameters.GetMutableArrayRef("ArrayValue").GetValue();
+				FInstancedPropertyBag& RootPropertyBag = GetRootPropertyBag(EditorData);
+				RootPropertyBag.AddProperty("Value", EPropertyBagPropertyType::Int32);
+				RootPropertyBag.SetValueInt32("Value", -111);
+				RootPropertyBag.AddContainerProperty("ArrayValue", FPropertyBagContainerTypes(EPropertyBagContainerType::Array), EPropertyBagPropertyType::Int32, nullptr);
+				RootPropertyBag.AddProperty("ArrayValue", EPropertyBagPropertyType::Int32);
+				FPropertyBagArrayRef ValueArrayRef = RootPropertyBag.GetMutableArrayRef("ArrayValue").GetValue();
 				ValueArrayRef.EmptyAndAddValues(4);
 				ValueArrayRef.SetValueInt32(0, -11);
 				ValueArrayRef.SetValueInt32(1, -22);
@@ -546,8 +547,8 @@ struct FStateTreeTest_MutableArray : FStateTreeTestBase
 				TStateTreeEditorNode<FTestTask_PrintValue>& TaskA = EditorData.AddGlobalTask<FTestTask_PrintValue>("Tree1GlobalTaskA");
 				TaskA.GetInstanceData().Value = -2;
 				TaskA.GetInstanceData().ArrayValue = {-1, -2};
-				EditorData.AddPropertyBinding(FPropertyBindingPath(EditorData.RootParameters.ID, TEXT("Value")), FPropertyBindingPath(TaskA.ID, TEXT("Value")));
-				EditorData.AddPropertyBinding(FPropertyBindingPath(EditorData.RootParameters.ID, TEXT("ArrayValue")), FPropertyBindingPath(TaskA.ID, TEXT("ArrayValue")));
+				EditorData.AddPropertyBinding(FPropertyBindingPath(EditorData.GetRootParametersGuid(), TEXT("Value")), FPropertyBindingPath(TaskA.ID, TEXT("Value")));
+				EditorData.AddPropertyBinding(FPropertyBindingPath(EditorData.GetRootParametersGuid(), TEXT("ArrayValue")), FPropertyBindingPath(TaskA.ID, TEXT("ArrayValue")));
 
 			}
 			UStateTreeState& Root = EditorData.AddSubTree("Tree1StateRoot");
@@ -558,7 +559,7 @@ struct FStateTreeTest_MutableArray : FStateTreeTestBase
 				Transition.bDelayTransition = true;
 				Transition.DelayDuration = 60.0;
 
-				FPropertyBindingPath RootParametersArrayValue3(EditorData.RootParameters.ID);
+				FPropertyBindingPath RootParametersArrayValue3(EditorData.GetRootParametersGuid());
 				RootParametersArrayValue3.AddPathSegment("ArrayValue", 3);
 
 				TStateTreeEditorNode<FTestTask_PrintAndResetValue>& TaskA = State.AddTask<FTestTask_PrintAndResetValue>("Tree1StateATaskA");
@@ -570,7 +571,7 @@ struct FStateTreeTest_MutableArray : FStateTreeTestBase
 				FPropertyBindingPath TaskAArrayValue3(TaskA.ID);
 				TaskAArrayValue3.AddPathSegment("ArrayValue", 3);
 
-				EditorData.AddPropertyBinding(FPropertyBindingPath(EditorData.RootParameters.ID, "Value"), FPropertyBindingPath(TaskA.ID, TEXT("Value")));
+				EditorData.AddPropertyBinding(FPropertyBindingPath(EditorData.GetRootParametersGuid(), "Value"), FPropertyBindingPath(TaskA.ID, TEXT("Value")));
 				EditorData.AddPropertyBinding(RootParametersArrayValue3, TaskAArrayValue3);
 				
 				TStateTreeEditorNode<FTestTask_PrintAndResetValue>& TaskB = State.AddTask<FTestTask_PrintAndResetValue>("Tree1StateATaskB");
@@ -643,35 +644,36 @@ struct FStateTreeTest_TransitionTaskWithBinding : FStateTreeTestBase
 			UStateTreeEditorData& EditorData = *Cast<UStateTreeEditorData>(StateTree.EditorData);
 			{
 				// Parameters
-				EditorData.RootParameters.Parameters.AddProperty("Value", EPropertyBagPropertyType::Int32);
-				EditorData.RootParameters.Parameters.SetValueInt32("Value", -111);
+				FInstancedPropertyBag& RootPropertyBag = GetRootPropertyBag(EditorData);
+				RootPropertyBag.AddProperty("Value", EPropertyBagPropertyType::Int32);
+				RootPropertyBag.SetValueInt32("Value", -111);
 
 				// Global
 				TStateTreeEditorNode<FTestTask_PrintValue>& TaskA = EditorData.AddGlobalTask<FTestTask_PrintValue>("Tree1GlobalTaskA");
 				TaskA.GetInstanceData().Value = -2;
-				EditorData.AddPropertyBinding(FPropertyBindingPath(EditorData.RootParameters.ID, TEXT("Value")), FPropertyBindingPath(TaskA.ID, TEXT("Value")));
+				EditorData.AddPropertyBinding(FPropertyBindingPath(EditorData.GetRootParametersGuid(), TEXT("Value")), FPropertyBindingPath(TaskA.ID, TEXT("Value")));
 
 				TStateTreeEditorNode<FTestTask_PrintValue_TransitionTick>& TaskB = EditorData.AddGlobalTask<FTestTask_PrintValue_TransitionTick>("Tree1GlobalTaskB");
 				TaskB.GetInstanceData().Value = -2;
-				EditorData.AddPropertyBinding(FPropertyBindingPath(EditorData.RootParameters.ID, TEXT("Value")), FPropertyBindingPath(TaskB.ID, TEXT("Value")));
+				EditorData.AddPropertyBinding(FPropertyBindingPath(EditorData.GetRootParametersGuid(), TEXT("Value")), FPropertyBindingPath(TaskB.ID, TEXT("Value")));
 
 				TStateTreeEditorNode<FTestTask_PrintValue_TransitionNoTick>& TaskC = EditorData.AddGlobalTask<FTestTask_PrintValue_TransitionNoTick>("Tree1GlobalTaskC");
 				TaskC.GetInstanceData().Value = -2;
-				EditorData.AddPropertyBinding(FPropertyBindingPath(EditorData.RootParameters.ID, TEXT("Value")), FPropertyBindingPath(TaskC.ID, TEXT("Value")));
+				EditorData.AddPropertyBinding(FPropertyBindingPath(EditorData.GetRootParametersGuid(), TEXT("Value")), FPropertyBindingPath(TaskC.ID, TEXT("Value")));
 			}
 			UStateTreeState& Root = EditorData.AddSubTree("Tree1StateRoot");
 			{
 				TStateTreeEditorNode<FTestTask_PrintValue>& TaskA = Root.AddTask<FTestTask_PrintValue>("Tree1StateRootTaskA");
 				TaskA.GetInstanceData().Value = -2;
-				EditorData.AddPropertyBinding(FPropertyBindingPath(EditorData.RootParameters.ID, TEXT("Value")), FPropertyBindingPath(TaskA.ID, TEXT("Value")));
+				EditorData.AddPropertyBinding(FPropertyBindingPath(EditorData.GetRootParametersGuid(), TEXT("Value")), FPropertyBindingPath(TaskA.ID, TEXT("Value")));
 				
 				TStateTreeEditorNode<FTestTask_PrintValue_TransitionTick>& TaskB = Root.AddTask<FTestTask_PrintValue_TransitionTick>("Tree1StateRootTaskB");
 				TaskB.GetInstanceData().Value = -2;
-				EditorData.AddPropertyBinding(FPropertyBindingPath(EditorData.RootParameters.ID, TEXT("Value")), FPropertyBindingPath(TaskB.ID, TEXT("Value")));
+				EditorData.AddPropertyBinding(FPropertyBindingPath(EditorData.GetRootParametersGuid(), TEXT("Value")), FPropertyBindingPath(TaskB.ID, TEXT("Value")));
 				
 				TStateTreeEditorNode<FTestTask_PrintValue_TransitionNoTick>& TaskC = Root.AddTask<FTestTask_PrintValue_TransitionNoTick>("Tree1StateRootTaskC");
 				TaskC.GetInstanceData().Value = -2;
-				EditorData.AddPropertyBinding(FPropertyBindingPath(EditorData.RootParameters.ID, TEXT("Value")), FPropertyBindingPath(TaskC.ID, TEXT("Value")));
+				EditorData.AddPropertyBinding(FPropertyBindingPath(EditorData.GetRootParametersGuid(), TEXT("Value")), FPropertyBindingPath(TaskC.ID, TEXT("Value")));
 			}
 			{
 				UStateTreeState& State = Root.AddChildState("Tree1StateA", EStateTreeStateType::State);
@@ -682,15 +684,15 @@ struct FStateTreeTest_TransitionTaskWithBinding : FStateTreeTestBase
 
 				TStateTreeEditorNode<FTestTask_PrintValue>& TaskA = State.AddTask<FTestTask_PrintValue>("Tree1StateATaskA");
 				TaskA.GetInstanceData().Value = -2;
-				EditorData.AddPropertyBinding(FPropertyBindingPath(EditorData.RootParameters.ID, TEXT("Value")), FPropertyBindingPath(TaskA.ID, TEXT("Value")));
+				EditorData.AddPropertyBinding(FPropertyBindingPath(EditorData.GetRootParametersGuid(), TEXT("Value")), FPropertyBindingPath(TaskA.ID, TEXT("Value")));
 
 				TStateTreeEditorNode<FTestTask_PrintValue_TransitionTick>& TaskB = State.AddTask<FTestTask_PrintValue_TransitionTick>("Tree1StateATaskB");
 				TaskB.GetInstanceData().Value = -2;
-				EditorData.AddPropertyBinding(FPropertyBindingPath(EditorData.RootParameters.ID, TEXT("Value")), FPropertyBindingPath(TaskB.ID, TEXT("Value")));
+				EditorData.AddPropertyBinding(FPropertyBindingPath(EditorData.GetRootParametersGuid(), TEXT("Value")), FPropertyBindingPath(TaskB.ID, TEXT("Value")));
 
 				TStateTreeEditorNode<FTestTask_PrintValue_TransitionNoTick>& TaskC = State.AddTask<FTestTask_PrintValue_TransitionNoTick>("Tree1StateATaskC");
 				TaskC.GetInstanceData().Value = -2;
-				EditorData.AddPropertyBinding(FPropertyBindingPath(EditorData.RootParameters.ID, TEXT("Value")), FPropertyBindingPath(TaskC.ID, TEXT("Value")));
+				EditorData.AddPropertyBinding(FPropertyBindingPath(EditorData.GetRootParametersGuid(), TEXT("Value")), FPropertyBindingPath(TaskC.ID, TEXT("Value")));
 			}
 		}
 		{
diff --git a/Engine/Plugins/Runtime/StateTree/Source/StateTreeTestSuite/Private/StateTreeLinkedStateTest.cpp b/Engine/Plugins/Runtime/StateTree/Source/StateTreeTestSuite/Private/StateTreeLinkedStateTest.cpp
--- a/Engine/Plugins/Runtime/StateTree/Source/StateTreeTestSuite/Private/StateTreeLinkedStateTest.cpp
+++ b/Engine/Plugins/Runtime/StateTree/Source/StateTreeTestSuite/Private/StateTreeLinkedStateTest.cpp
@@ -205,20 +205,21 @@ struct FStateTreeTest_MultipleSameLinkedAsset : FStateTreeTestBase
 			UStateTreeEditorData& EditorData2 = *Cast<UStateTreeEditorData>(StateTree2.EditorData);
 			{
 				// Parameters
-				EditorData2.RootParameters.Parameters.AddProperty("Value", EPropertyBagPropertyType::Int32);
-				EditorData2.RootParameters.Parameters.SetValueInt32("Value", -111);
-				RootParameter_ValueID = EditorData2.RootParameters.Parameters.FindPropertyDescByName("Value")->ID;
+				FInstancedPropertyBag& RootPropertyBag = GetRootPropertyBag(EditorData2);
+				RootPropertyBag.AddProperty("Value", EPropertyBagPropertyType::Int32);
+				RootPropertyBag.SetValueInt32("Value", -111);
+				RootParameter_ValueID = RootPropertyBag.FindPropertyDescByName("Value")->ID;
 
 				TStateTreeEditorNode<FTestTask_PrintValue>& GlobalTask = EditorData2.AddGlobalTask<FTestTask_PrintValue>("Tree2GlobalTaskA");
 				GlobalTask.GetInstanceData().Value = -1;
-				EditorData2.AddPropertyBinding(FPropertyBindingPath(EditorData2.RootParameters.ID, TEXT("Value")), FPropertyBindingPath(GlobalTask.ID, TEXT("Value")));
+				EditorData2.AddPropertyBinding(FPropertyBindingPath(EditorData2.GetRootParametersGuid(), TEXT("Value")), FPropertyBindingPath(GlobalTask.ID, TEXT("Value")));
 			}
 
 			UStateTreeState& Root = EditorData2.AddSubTree("Tree2StateRoot");
 			{
 				TStateTreeEditorNode<FTestTask_PrintValue>& Task1 = Root.AddTask<FTestTask_PrintValue>("Tree2StateRootTaskA");
 				Task1.GetInstanceData().Value = -2;
-				EditorData2.AddPropertyBinding(FPropertyBindingPath(EditorData2.RootParameters.ID, TEXT("Value")), FPropertyBindingPath(Task1.ID, TEXT("Value")));
+				EditorData2.AddPropertyBinding(FPropertyBindingPath(EditorData2.GetRootParametersGuid(), TEXT("Value")), FPropertyBindingPath(Task1.ID, TEXT("Value")));
 			}
 			{
 				UStateTreeState& State = Root.AddChildState("Tree2StateA", EStateTreeStateType::State);
@@ -228,7 +229,7 @@ struct FStateTreeTest_MultipleSameLinkedAsset : FStateTreeTestBase
 
 				TStateTreeEditorNode<FTestTask_PrintValue>& Task1 = State.AddTask<FTestTask_PrintValue>("Tree2StateATaskA");
 				Task1.GetInstanceData().Value = -2;
-				EditorData2.AddPropertyBinding(FPropertyBindingPath(EditorData2.RootParameters.ID, TEXT("Value")), FPropertyBindingPath(Task1.ID, TEXT("Value")));
+				EditorData2.AddPropertyBinding(FPropertyBindingPath(EditorData2.GetRootParametersGuid(), TEXT("Value")), FPropertyBindingPath(Task1.ID, TEXT("Value")));
 			}
 
 			FStateTreeCompiler Compiler2(Log);
diff --git a/Engine/Plugins/Runtime/StateTree/Source/StateTreeTestSuite/Private/StateTreeTest.cpp b/Engine/Plugins/Runtime/StateTree/Source/StateTreeTestSuite/Private/StateTreeTest.cpp
--- a/Engine/Plugins/Runtime/StateTree/Source/StateTreeTestSuite/Private/StateTreeTest.cpp
+++ b/Engine/Plugins/Runtime/StateTree/Source/StateTreeTestSuite/Private/StateTreeTest.cpp
@@ -1794,7 +1794,8 @@ struct FStateTreeTest_NestedOverride : FStateTreeTestBase
 		// Asset 2
 		UStateTree& StateTree2 = NewStateTree();
 		UStateTreeEditorData& EditorData2 = *Cast<UStateTreeEditorData>(StateTree2.EditorData);
-		EditorData2.RootParameters.Parameters.AddProperty(FName(TEXT("Int")), EPropertyBagPropertyType::Int32);
+		FInstancedPropertyBag& RootPropertyBag2 = GetRootPropertyBag(EditorData2);
+		RootPropertyBag2.AddProperty(FName(TEXT("Int")), EPropertyBagPropertyType::Int32);
 		UStateTreeState& Root2 = EditorData2.AddSubTree(FName(TEXT("Root2")));
 		TStateTreeEditorNode<FTestTask_Stand>& TaskRoot2 = Root2.AddTask<FTestTask_Stand>(FName(TEXT("TaskRoot2")));
 		{
@@ -1806,7 +1807,8 @@ struct FStateTreeTest_NestedOverride : FStateTreeTestBase
 		// Asset 3
 		UStateTree& StateTree3 = NewStateTree();
 		UStateTreeEditorData& EditorData3 = *Cast<UStateTreeEditorData>(StateTree3.EditorData);
-		EditorData3.RootParameters.Parameters.AddProperty(FName(TEXT("Float")), EPropertyBagPropertyType::Float); // Different parameters
+		FInstancedPropertyBag& RootPropertyBag3 = GetRootPropertyBag(EditorData3);
+		RootPropertyBag3.AddProperty(FName(TEXT("Float")), EPropertyBagPropertyType::Float); // Different parameters
 		UStateTreeState& Root3 = EditorData3.AddSubTree(FName(TEXT("Root3")));
 		TStateTreeEditorNode<FTestTask_Stand>& TaskRoot3 = Root3.AddTask<FTestTask_Stand>(FName(TEXT("TaskRoot3")));
 		{
@@ -1821,8 +1823,9 @@ struct FStateTreeTest_NestedOverride : FStateTreeTestBase
 			check(EditorData);
 			StateTree4->EditorData = EditorData;
 			EditorData->Schema = NewObject<UStateTreeTestSchema2>();
-			
-			EditorData->RootParameters.Parameters.AddProperty(FName(TEXT("Float")), EPropertyBagPropertyType::Float); // Different parameters
+
+			FInstancedPropertyBag& RootPropertyBag = GetRootPropertyBag(*EditorData);
+			RootPropertyBag.AddProperty(FName(TEXT("Float")), EPropertyBagPropertyType::Float); // Different parameters
 			UStateTreeState& Root4 = EditorData->AddSubTree(FName(TEXT("Root4")));
 			TStateTreeEditorNode<FTestTask_Stand>& TaskRoot4 = Root3.AddTask<FTestTask_Stand>(FName(TEXT("TaskRoot4")));
 
@@ -1834,8 +1837,9 @@ struct FStateTreeTest_NestedOverride : FStateTreeTestBase
 		// Main asset
 		UStateTree& StateTree = NewStateTree();
 		UStateTreeEditorData& EditorData = *Cast<UStateTreeEditorData>(StateTree.EditorData);
+		FInstancedPropertyBag& RootPropertyBag = GetRootPropertyBag(EditorData);
 
-		EditorData.RootParameters.Parameters.AddProperty(FName(TEXT("Int")), EPropertyBagPropertyType::Int32);
+		RootPropertyBag.AddProperty(FName(TEXT("Int")), EPropertyBagPropertyType::Int32);
 		
 		UStateTreeState& Root = EditorData.AddSubTree(FName(TEXT("Root1")));
 		UStateTreeState& StateA = Root.AddChildState(FName(TEXT("A1")), EStateTreeStateType::LinkedAsset);
diff --git a/Engine/Plugins/Runtime/StateTree/Source/StateTreeTestSuite/Private/StateTreeTestBase.cpp b/Engine/Plugins/Runtime/StateTree/Source/StateTreeTestSuite/Private/StateTreeTestBase.cpp
--- a/Engine/Plugins/Runtime/StateTree/Source/StateTreeTestSuite/Private/StateTreeTestBase.cpp
+++ b/Engine/Plugins/Runtime/StateTree/Source/StateTreeTestSuite/Private/StateTreeTestBase.cpp
@@ -69,6 +69,10 @@ FGameplayTag FStateTreeTestBase::GetTestTag3()
 	return GameplayTagStaticInstance.TestTag3;
 }
 
+FInstancedPropertyBag& FStateTreeTestBase::GetRootPropertyBag(UStateTreeEditorData& EditorData) const
+{
+	return const_cast<FInstancedPropertyBag&>(EditorData.GetRootParametersPropertyBag());
+}
 }//namespace UE::StateTree::Tests
 
 #undef LOCTEXT_NAMESPACE
diff --git a/Engine/Plugins/Runtime/StateTree/Source/StateTreeTestSuite/Private/StateTreeTransitionTest.cpp b/Engine/Plugins/Runtime/StateTree/Source/StateTreeTestSuite/Private/StateTreeTransitionTest.cpp
--- a/Engine/Plugins/Runtime/StateTree/Source/StateTreeTestSuite/Private/StateTreeTransitionTest.cpp
+++ b/Engine/Plugins/Runtime/StateTree/Source/StateTreeTestSuite/Private/StateTreeTransitionTest.cpp
@@ -608,9 +608,9 @@ struct FStateTreeTest_FollowTransitions : FStateTreeTestBase
 	{
 		UStateTree& StateTree = NewStateTree();
 		UStateTreeEditorData& EditorData = *Cast<UStateTreeEditorData>(StateTree.EditorData);
-
-		EditorData.RootParameters.Parameters.AddProperty(FName(TEXT("Int")), EPropertyBagPropertyType::Int32);
-		EditorData.RootParameters.Parameters.SetValueInt32(FName(TEXT("Int")), 1);
+		FInstancedPropertyBag& RootPropertyBag = GetRootPropertyBag(EditorData);
+		RootPropertyBag.AddProperty(FName(TEXT("Int")), EPropertyBagPropertyType::Int32);
+		RootPropertyBag.SetValueInt32(FName(TEXT("Int")), 1);
 
 		UStateTreeState& Root = EditorData.AddSubTree(FName(TEXT("Root")));
 		UStateTreeState& StateTrans = Root.AddChildState(FName(TEXT("Trans")));
@@ -630,7 +630,7 @@ struct FStateTreeTest_FollowTransitions : FStateTreeTestBase
 				TStateTreeEditorNode<FStateTreeCompareIntCondition>& TransIntCond = TransA.AddCondition<FStateTreeCompareIntCondition>(EGenericAICheck::Equal);
 				TransIntCond.GetInstanceData().Right = 0;
 				EditorData.AddPropertyBinding(
-					FPropertyBindingPath(EditorData.RootParameters.ID, TEXT("Int")),
+					FPropertyBindingPath(EditorData.GetRootParametersGuid(), TEXT("Int")),
 					FPropertyBindingPath(TransIntCond.ID, TEXT("Left")));
 			}
 
@@ -641,7 +641,7 @@ struct FStateTreeTest_FollowTransitions : FStateTreeTestBase
 				TStateTreeEditorNode<FStateTreeCompareIntCondition>& TransIntCond = TransB.AddCondition<FStateTreeCompareIntCondition>(EGenericAICheck::Equal);
 				TransIntCond.GetInstanceData().Right = 1;
 				EditorData.AddPropertyBinding(
-					FPropertyBindingPath(EditorData.RootParameters.ID, TEXT("Int")),
+					FPropertyBindingPath(EditorData.GetRootParametersGuid(), TEXT("Int")),
 					FPropertyBindingPath(TransIntCond.ID, TEXT("Left")));
 			}
 
@@ -652,7 +652,7 @@ struct FStateTreeTest_FollowTransitions : FStateTreeTestBase
 				TStateTreeEditorNode<FStateTreeCompareIntCondition>& TransIntCond = TransC.AddCondition<FStateTreeCompareIntCondition>(EGenericAICheck::Equal);
 				TransIntCond.GetInstanceData().Right = 1;
 				EditorData.AddPropertyBinding(
-					FPropertyBindingPath(EditorData.RootParameters.ID, TEXT("Int")),
+					FPropertyBindingPath(EditorData.GetRootParametersGuid(), TEXT("Int")),
 					FPropertyBindingPath(TransIntCond.ID, TEXT("Left")));
 			}
 		}
@@ -696,9 +696,9 @@ struct FStateTreeTest_InfiniteLoop : FStateTreeTestBase
 	{
 		UStateTree& StateTree = NewStateTree();
 		UStateTreeEditorData& EditorData = *Cast<UStateTreeEditorData>(StateTree.EditorData);
-
-		EditorData.RootParameters.Parameters.AddProperty(FName(TEXT("Int")), EPropertyBagPropertyType::Int32);
-		EditorData.RootParameters.Parameters.SetValueInt32(FName(TEXT("Int")), 1);
+		FInstancedPropertyBag& RootPropertyBag = GetRootPropertyBag(EditorData);
+		RootPropertyBag.AddProperty(FName(TEXT("Int")), EPropertyBagPropertyType::Int32);
+		RootPropertyBag.SetValueInt32(FName(TEXT("Int")), 1);
 		
 		UStateTreeState& Root = EditorData.AddSubTree(FName(TEXT("Root")));
 		UStateTreeState& StateA = Root.AddChildState(FName(TEXT("A")));
@@ -715,7 +715,7 @@ struct FStateTreeTest_InfiniteLoop : FStateTreeTestBase
 				TStateTreeEditorNode<FStateTreeCompareIntCondition>& TransIntCond = Trans.AddCondition<FStateTreeCompareIntCondition>(EGenericAICheck::Equal);
 				TransIntCond.GetInstanceData().Right = 1;
 				EditorData.AddPropertyBinding(
-					FPropertyBindingPath(EditorData.RootParameters.ID, TEXT("Int")),
+					FPropertyBindingPath(EditorData.GetRootParametersGuid(), TEXT("Int")),
 					FPropertyBindingPath(TransIntCond.ID, TEXT("Left")));
 			}
 		}
@@ -729,7 +729,7 @@ struct FStateTreeTest_InfiniteLoop : FStateTreeTestBase
 				TStateTreeEditorNode<FStateTreeCompareIntCondition>& TransIntCond = Trans.AddCondition<FStateTreeCompareIntCondition>(EGenericAICheck::Equal);
 				TransIntCond.GetInstanceData().Right = 1;
 				EditorData.AddPropertyBinding(
-					FPropertyBindingPath(EditorData.RootParameters.ID, TEXT("Int")),
+					FPropertyBindingPath(EditorData.GetRootParametersGuid(), TEXT("Int")),
 					FPropertyBindingPath(TransIntCond.ID, TEXT("Left")));
 			}
 		}
@@ -790,20 +790,21 @@ struct FStateTreeTest_RegularTransitions : FStateTreeTestBase
 			FGuid RootParameter_ValueID;
 			{
 				// Parameters
-				EditorData.RootParameters.Parameters.AddProperty("Value", EPropertyBagPropertyType::Int32);
-				EditorData.RootParameters.Parameters.SetValueInt32("Value", -111);
-				RootParameter_ValueID = EditorData.RootParameters.Parameters.FindPropertyDescByName("Value")->ID;
+				FInstancedPropertyBag& RootPropertyBag = GetRootPropertyBag(EditorData);
+				RootPropertyBag.AddProperty("Value", EPropertyBagPropertyType::Int32);
+				RootPropertyBag.SetValueInt32("Value", -111);
+				RootParameter_ValueID = RootPropertyBag.FindPropertyDescByName("Value")->ID;
 
 				TStateTreeEditorNode<FTestTask_PrintValue>& GlobalTask = EditorData.AddGlobalTask<FTestTask_PrintValue>("GlobalTask");
 				GlobalTask.GetInstanceData().Value = -1;
-				EditorData.AddPropertyBinding(FPropertyBindingPath(EditorData.RootParameters.ID, TEXT("Value")), FPropertyBindingPath(GlobalTask.ID, TEXT("Value")));
+				EditorData.AddPropertyBinding(FPropertyBindingPath(EditorData.GetRootParametersGuid(), TEXT("Value")), FPropertyBindingPath(GlobalTask.ID, TEXT("Value")));
 			}
 
 			UStateTreeState& Root = EditorData.AddSubTree("RootA");
 			{
 				TStateTreeEditorNode<FTestTask_PrintValue>& Task = Root.AddTask<FTestTask_PrintValue>("TaskA");
 				Task.GetInstanceData().Value = -1;
-				EditorData.AddPropertyBinding(FPropertyBindingPath(EditorData.RootParameters.ID, TEXT("Value")), FPropertyBindingPath(Task.ID, TEXT("Value")));
+				EditorData.AddPropertyBinding(FPropertyBindingPath(EditorData.GetRootParametersGuid(), TEXT("Value")), FPropertyBindingPath(Task.ID, TEXT("Value")));
 			}
 			{
 				UStateTreeState& StateB = Root.AddChildState("StateB", EStateTreeStateType::State);
diff --git a/Engine/Plugins/Runtime/StateTree/Source/StateTreeTestSuite/Public/StateTreeTestBase.h b/Engine/Plugins/Runtime/StateTree/Source/StateTreeTestSuite/Public/StateTreeTestBase.h
--- a/Engine/Plugins/Runtime/StateTree/Source/StateTreeTestSuite/Public/StateTreeTestBase.h
+++ b/Engine/Plugins/Runtime/StateTree/Source/StateTreeTestSuite/Public/StateTreeTestBase.h
@@ -7,6 +7,7 @@
 #include "StateTree.h"
 #include "StateTreePropertyBindings.h"
 
+class UStateTreeEditorData;
 namespace UE::StateTree::Tests
 {
 	
@@ -21,6 +22,7 @@ protected:
 	static FGameplayTag GetTestTag1();
 	static FGameplayTag GetTestTag2();
 	static FGameplayTag GetTestTag3();
+	FInstancedPropertyBag& GetRootPropertyBag(UStateTreeEditorData& EditorData) const;
 };
 
 } // namespace UE::StateTree::Tests
\ No newline at end of file
diff --git a/Engine/Source/Runtime/Core/Public/UObject/UE5SpecialProjectStreamObjectVersion.inl b/Engine/Source/Runtime/Core/Public/UObject/UE5SpecialProjectStreamObjectVersion.inl
--- a/Engine/Source/Runtime/Core/Public/UObject/UE5SpecialProjectStreamObjectVersion.inl
+++ b/Engine/Source/Runtime/Core/Public/UObject/UE5SpecialProjectStreamObjectVersion.inl
@@ -33,5 +33,8 @@ UE5_SPECIAL_PROJECT_VERSION(SerializeSkeletalMeshMorphTargetRenderData, 7)
 // Strip the Morph Target source data for cooked builds
 UE5_SPECIAL_PROJECT_VERSION(StripMorphTargetSourceDataForCookedBuilds, 8)
 
+// StateTree now holds PropertyBag + GUID for root-level parameters rather than FStateTreeStateParameters. Access is protected by default and can be overriden through virtuals on UStateTreeEditorData derived classes.
+UE5_SPECIAL_PROJECT_VERSION(StateTreeGlobalParameterChanges, 9)
+
 // -----<new versions can be added above this line>-------------------------------------------------
 #undef UE5_SPECIAL_PROJECT_VERSION
\ No newline at end of file
