Propagate all struct changes to ST
- Also fix race condition where AnimNext would not initialize StateTree due to non-deterministic module init order.

diff --git a/Engine/Plugins/Experimental/AnimNextStateTree/Source/AnimNextStateTreeEditor/Private/AnimNextStateTreeEditorModule.cpp b/Engine/Plugins/Experimental/AnimNextStateTree/Source/AnimNextStateTreeEditor/Private/AnimNextStateTreeEditorModule.cpp
--- a/Engine/Plugins/Experimental/AnimNextStateTree/Source/AnimNextStateTreeEditor/Private/AnimNextStateTreeEditorModule.cpp
+++ b/Engine/Plugins/Experimental/AnimNextStateTree/Source/AnimNextStateTreeEditor/Private/AnimNextStateTreeEditorModule.cpp
@@ -321,11 +321,9 @@ void FAnimNextStateTreeEditorModule::StartupModule()
 	TSharedPtr<FStateTreeAssetOutlinerItemDetails> StateItemDetails = MakeShareable<FStateTreeAssetOutlinerItemDetails>(new FStateTreeAssetOutlinerItemDetails());
 	WorkspaceEditorModule.RegisterWorkspaceItemDetails(Workspace::FOutlinerItemDetailsId(FAnimNextStateTreeOutlinerData::StaticStruct()->GetFName()), StaticCastSharedPtr<UE::Workspace::IWorkspaceOutlinerItemDetails>(StateItemDetails));	
 	WorkspaceEditorModule.RegisterWorkspaceItemDetails(Workspace::FOutlinerItemDetailsId(FAnimNextStateTreeStateOutlinerData::StaticStruct()->GetFName()), StaticCastSharedPtr<UE::Workspace::IWorkspaceOutlinerItemDetails>(StateItemDetails));
-	
-	if(FStateTreeEditorModule* StateTreeEditorModule = FStateTreeEditorModule::GetModulePtr())
-	{
-		StateTreeEditorModule->RegisterEditorDataClass(UStateTreeAnimNextSchema::StaticClass(), UAnimNextStateTreeTreeEditorData::StaticClass());		
-	}
+
+	FStateTreeEditorModule& StateTreeEditorModule = FStateTreeEditorModule::GetModule();
+	StateTreeEditorModule.RegisterEditorDataClass(UStateTreeAnimNextSchema::StaticClass(), UAnimNextStateTreeTreeEditorData::StaticClass());		
 }
 
 void FAnimNextStateTreeEditorModule::ShutdownModule()
diff --git a/Engine/Plugins/Experimental/AnimNextStateTree/Source/AnimNextStateTreeUncookedOnly/Internal/AnimNextStateTreeEditorData.cpp b/Engine/Plugins/Experimental/AnimNextStateTree/Source/AnimNextStateTreeUncookedOnly/Internal/AnimNextStateTreeEditorData.cpp
--- a/Engine/Plugins/Experimental/AnimNextStateTree/Source/AnimNextStateTreeUncookedOnly/Internal/AnimNextStateTreeEditorData.cpp
+++ b/Engine/Plugins/Experimental/AnimNextStateTree/Source/AnimNextStateTreeUncookedOnly/Internal/AnimNextStateTreeEditorData.cpp
@@ -25,9 +25,6 @@ void UAnimNextStateTreeTreeEditorData::CreateRootProperties(TArrayView<UE::Prope
 
 	UAnimNextStateTree_EditorData* AnimNextEditorData = UE::AnimNext::UncookedOnly::FUtils::GetEditorData<UAnimNextStateTree_EditorData>(OuterAnimNextStateTree);
 
-	TArray<FPropertyBagPropertyDesc, TInlineAllocator<1>> PropertyDescs;
-	PropertyDescs.Reserve(InOutCreationDescs.Num());
-
 	// Generate unique names for the incoming property descs to avoid changing the existing properties in the bag
 	for (UE::PropertyBinding::FPropertyCreationDescriptor& CreationDesc : InOutCreationDescs)
 	{
diff --git a/Engine/Plugins/Experimental/AnimNextStateTree/Source/AnimNextStateTreeUncookedOnly/Internal/AnimNextStateTree_EditorData.cpp b/Engine/Plugins/Experimental/AnimNextStateTree/Source/AnimNextStateTreeUncookedOnly/Internal/AnimNextStateTree_EditorData.cpp
--- a/Engine/Plugins/Experimental/AnimNextStateTree/Source/AnimNextStateTreeUncookedOnly/Internal/AnimNextStateTree_EditorData.cpp
+++ b/Engine/Plugins/Experimental/AnimNextStateTree/Source/AnimNextStateTreeUncookedOnly/Internal/AnimNextStateTree_EditorData.cpp
@@ -10,6 +10,8 @@
 #include "Compilation/AnimNextGetGraphCompileContext.h"
 #include "DataInterface/AnimNextDataInterface_EditorData.h"
 #include "StateTree.h"
+#include "StateTreeCompilerLog.h"
+#include "StateTreeEditingSubsystem.h"
 #include "StateTreeEditorData.h"
 #include "Entries/AnimNextVariableEntry.h"
 #include "Entries/AnimNextAnimationGraphEntry.h"
@@ -33,6 +35,22 @@ TSubclassOf<UAssetUserData> UAnimNextStateTree_EditorData::GetAssetUserDataClass
 	return UAnimNextStateTreeWorkspaceAssetUserData::StaticClass();
 }
 
+void UAnimNextStateTree_EditorData::RecompileVM()
+{
+	UAnimNextAnimationGraph_EditorData::RecompileVM();
+
+	UAnimNextStateTree* AnimationStateTree = UE::AnimNext::UncookedOnly::FUtils::GetAsset<UAnimNextStateTree>(this);
+
+	UStateTree* InnerStateTree = AnimationStateTree->StateTree;
+
+	UStateTreeEditorData* InnerEditorData = Cast<UStateTreeEditorData>(InnerStateTree->EditorData);
+
+	// Recompile ST as we just updated our variables propertybag. It may have changed from an external data interface update.
+	FStateTreeCompilerLog Log;
+	bool bCompileSucceeded = UStateTreeEditingSubsystem::CompileStateTree(InnerStateTree, Log);
+	ensureMsgf(bCompileSucceeded, TEXT("Failed to compile state tree after data update: %s"), *InnerStateTree->GetFName().ToString());
+}
+
 TConstArrayView<TSubclassOf<UAnimNextRigVMAssetEntry>> UAnimNextStateTree_EditorData::GetEntryClasses() const
 {
 	static const TSubclassOf<UAnimNextRigVMAssetEntry> Classes[] =
diff --git a/Engine/Plugins/Experimental/AnimNextStateTree/Source/AnimNextStateTreeUncookedOnly/Internal/AnimNextStateTree_EditorData.h b/Engine/Plugins/Experimental/AnimNextStateTree/Source/AnimNextStateTreeUncookedOnly/Internal/AnimNextStateTree_EditorData.h
--- a/Engine/Plugins/Experimental/AnimNextStateTree/Source/AnimNextStateTreeUncookedOnly/Internal/AnimNextStateTree_EditorData.h
+++ b/Engine/Plugins/Experimental/AnimNextStateTree/Source/AnimNextStateTreeUncookedOnly/Internal/AnimNextStateTree_EditorData.h
@@ -16,6 +16,9 @@ protected:
 
 	friend class UAnimNextStateTreeFactory;
 
+	// IRigVMClientHost interface
+	virtual void RecompileVM() override;
+	
 	// UAnimNextRigVMAssetEditorData interface
 	virtual TConstArrayView<TSubclassOf<UAnimNextRigVMAssetEntry>> GetEntryClasses() const override;
 
