Revert layout & type based comparision from statetree / PropertyBindingCollection

diff --git a/Engine/Plugins/Experimental/AnimNextStateTree/Source/AnimNextStateTree/Private/AnimStateTreeTrait.cpp b/Engine/Plugins/Experimental/AnimNextStateTree/Source/AnimNextStateTree/Private/AnimStateTreeTrait.cpp
--- a/Engine/Plugins/Experimental/AnimNextStateTree/Source/AnimNextStateTree/Private/AnimStateTreeTrait.cpp
+++ b/Engine/Plugins/Experimental/AnimNextStateTree/Source/AnimNextStateTree/Private/AnimStateTreeTrait.cpp
@@ -103,7 +103,7 @@ namespace UE::AnimNext
 			StateTreeExecutionContext.SetContextDataByName(UStateTreeAnimNextSchema::AnimStateTreeExecutionContextName, FStateTreeDataView(FAnimNextStateTreeTraitContext::StaticStruct(), reinterpret_cast<uint8*>(&TraitContext)));
 
 			FAnimNextGraphInstance* OwnerGraphInstance = &Binding.GetTraitPtr().GetNodeInstance()->GetOwner();
-			const bool bPropertyLayoutMatches = InstanceData->StateTree->GetDefaultParameters().HasSameLayout(OwnerGraphInstance->GetVariables());
+			const bool bPropertyLayoutMatches = InstanceData->StateTree->GetDefaultParameters().GetPropertyBagStruct() == OwnerGraphInstance->GetVariables().GetPropertyBagStruct();
 
 			if(OwnerGraphInstance && ensure(bPropertyLayoutMatches))
 			{
diff --git a/Engine/Plugins/Runtime/PropertyBindingUtils/Source/PropertyBindingUtils/Private/PropertyBindingBindingCollection.cpp b/Engine/Plugins/Runtime/PropertyBindingUtils/Source/PropertyBindingUtils/Private/PropertyBindingBindingCollection.cpp
--- a/Engine/Plugins/Runtime/PropertyBindingUtils/Source/PropertyBindingUtils/Private/PropertyBindingBindingCollection.cpp
+++ b/Engine/Plugins/Runtime/PropertyBindingUtils/Source/PropertyBindingUtils/Private/PropertyBindingBindingCollection.cpp
@@ -1080,7 +1080,7 @@ bool FPropertyBindingBindingCollection::CopyProperty(const FPropertyBindingCopyI
 	
 	if (SourceStructView.IsValid() && TargetStructView.IsValid())
 	{
-		check(FStructUtils::TheSameLayout(SourceStructView.GetStruct(), Copy.SourceStructType) 
+		check(SourceStructView.GetStruct() == Copy.SourceStructType
 			|| (SourceStructView.GetStruct() && SourceStructView.GetStruct()->IsChildOf(Copy.SourceStructType)));
 			
 		uint8* SourceAddress = GetAddress(SourceStructView, Copy.SourceIndirection, Copy.SourceLeafProperty);
diff --git a/Engine/Plugins/Runtime/StateTree/Source/StateTreeModule/Private/StateTreeExecutionContext.cpp b/Engine/Plugins/Runtime/StateTree/Source/StateTreeModule/Private/StateTreeExecutionContext.cpp
--- a/Engine/Plugins/Runtime/StateTree/Source/StateTreeModule/Private/StateTreeExecutionContext.cpp
+++ b/Engine/Plugins/Runtime/StateTree/Source/StateTreeModule/Private/StateTreeExecutionContext.cpp
@@ -2600,7 +2600,7 @@ FStateTreeIndex16 FStateTreeExecutionContext::CollectExternalData(const UStateTr
 
 bool FStateTreeExecutionContext::SetGlobalParameters(const FInstancedPropertyBag& Parameters)
 {
-	if (ensureMsgf(RootStateTree.GetDefaultParameters().HasSameLayout(Parameters),
+	if (ensureMsgf(RootStateTree.GetDefaultParameters().GetPropertyBagStruct() == Parameters.GetPropertyBagStruct(),
 		TEXT("Parameters must be of the same struct type. Make sure to migrate the provided parameters to the same type as the StateTree default parameters.")))
 	{
 		Storage.SetGlobalParameters(Parameters);
