StateTree: Remove FStateTreeDetailsCustomizationSettings and use AllowGlobalParameters instead.
#rnx

diff --git a/Engine/Plugins/Experimental/AnimNextStateTree/Source/AnimNextStateTree/Private/AnimNextStateTreeSchema.cpp b/Engine/Plugins/Experimental/AnimNextStateTree/Source/AnimNextStateTree/Private/AnimNextStateTreeSchema.cpp
--- a/Engine/Plugins/Experimental/AnimNextStateTree/Source/AnimNextStateTree/Private/AnimNextStateTreeSchema.cpp
+++ b/Engine/Plugins/Experimental/AnimNextStateTree/Source/AnimNextStateTree/Private/AnimNextStateTreeSchema.cpp
@@ -19,9 +19,6 @@ const FName UStateTreeAnimNextSchema::AnimStateTreeExecutionContextName("Executi
 UStateTreeAnimNextSchema::UStateTreeAnimNextSchema() : ContextDataDescs(
 	{{ AnimStateTreeExecutionContextName, FAnimNextStateTreeTraitContext::StaticStruct(), FGuid(0xDFB93B7F, 0xEDBE4906, 0x851C66B2, 0x7585FA21) }})
 {
-#if WITH_EDITOR
-	DetailsCustomizationSettings.bShowStateTreeParameters = false;
-#endif // WITH_EDITOR
 }
 
 bool UStateTreeAnimNextSchema::IsStructAllowed(const UScriptStruct* InScriptStruct) const
diff --git a/Engine/Plugins/Experimental/AnimNextStateTree/Source/AnimNextStateTree/Public/AnimNextStateTreeSchema.h b/Engine/Plugins/Experimental/AnimNextStateTree/Source/AnimNextStateTree/Public/AnimNextStateTreeSchema.h
--- a/Engine/Plugins/Experimental/AnimNextStateTree/Source/AnimNextStateTree/Public/AnimNextStateTreeSchema.h
+++ b/Engine/Plugins/Experimental/AnimNextStateTree/Source/AnimNextStateTree/Public/AnimNextStateTreeSchema.h
@@ -24,16 +24,22 @@ public:
 	virtual EStateTreeParameterDataType GetGlobalParameterDataType() const override;
 	
 #if WITH_EDITOR
-	virtual bool AllowUtilityConsiderations() const override { return false; }
-	virtual bool AllowEvaluators() const override { return false; }
-	virtual const FStateTreeDetailsCustomizationSettings& GetDetailsCustomizationSettings() const override { return DetailsCustomizationSettings; }
+	virtual bool AllowUtilityConsiderations() const override
+	{
+		return false;
+	}
+	virtual bool AllowEvaluators() const override
+	{
+		return false;
+	}
+	virtual bool AllowGlobalParameters() const override
+	{
+		return false;
+	}
 #endif // WITH_EDITOR
+
 protected:
 	/** List of named external data required by schema and provided to the state tree through the execution context. */
 	UPROPERTY()
 	TArray<FStateTreeExternalDataDesc> ContextDataDescs;
-
-#if WITH_EDITOR
-	FStateTreeDetailsCustomizationSettings DetailsCustomizationSettings;
-#endif // WITH_EDITOR
 };
diff --git a/Engine/Plugins/Runtime/StateTree/Source/StateTreeEditorModule/Private/Customizations/StateTreeEditorDataDetails.cpp b/Engine/Plugins/Runtime/StateTree/Source/StateTreeEditorModule/Private/Customizations/StateTreeEditorDataDetails.cpp
--- a/Engine/Plugins/Runtime/StateTree/Source/StateTreeEditorModule/Private/Customizations/StateTreeEditorDataDetails.cpp
+++ b/Engine/Plugins/Runtime/StateTree/Source/StateTreeEditorModule/Private/Customizations/StateTreeEditorDataDetails.cpp
@@ -40,7 +40,6 @@ void FStateTreeEditorDataDetails::CustomizeDetails(IDetailLayoutBuilder& DetailB
 		}
 	}
 	const UStateTreeSchema* Schema = EditorData ? EditorData->Schema : nullptr;
-	const FStateTreeDetailsCustomizationSettings* CustomizationSettings = Schema ? &Schema->GetDetailsCustomizationSettings() : nullptr;
 
 	// Common category
 	IDetailCategoryBuilder& CommonCategory = DetailBuilder.EditCategory(TEXT("Common"), LOCTEXT("EditorDataDetailsCommon", "Common"));
@@ -140,7 +139,7 @@ void FStateTreeEditorDataDetails::CustomizeDetails(IDetailLayoutBuilder& DetailB
 		}
 	}
 
-	if (CustomizationSettings && CustomizationSettings->bShowStateTreeParameters)
+	if (Schema && Schema->AllowGlobalParameters())
 	{
 		// Parameters category
 		const FText ParametersDisplayName = LOCTEXT("EditorDataDetailsParameters", "Parameters");
diff --git a/Engine/Plugins/Runtime/StateTree/Source/StateTreeModule/Public/StateTreeSchema.h b/Engine/Plugins/Runtime/StateTree/Source/StateTreeModule/Public/StateTreeSchema.h
--- a/Engine/Plugins/Runtime/StateTree/Source/StateTreeModule/Public/StateTreeSchema.h
+++ b/Engine/Plugins/Runtime/StateTree/Source/StateTreeModule/Public/StateTreeSchema.h
@@ -13,12 +13,6 @@ enum class EStateTreeDataSourceType : uint8;
 enum class EStateTreeParameterDataType : uint8;
 enum class EStateTreeStateSelectionBehavior : uint8;
 
-#if WITH_EDITOR
-struct FStateTreeDetailsCustomizationSettings
-{
-	bool bShowStateTreeParameters : 1 = true;
-};
-#endif // WITH_EDITOR
 
 /**
  * Schema describing which inputs, evaluators, and tasks a StateTree can contain.
@@ -113,18 +107,17 @@ public:
 		return true;
 	}
 
-	/** @return True if modifying the tasks completion is allowed. */
-	virtual bool AllowTasksCompletion() const
+	/** @return True if global parameters for the State Tree are allowed. */
+	virtual bool AllowGlobalParameters() const
 	{
 		return true;
 	}
 
-	virtual const FStateTreeDetailsCustomizationSettings& GetDetailsCustomizationSettings() const
+	/** @return True if modifying the tasks completion is allowed. */
+	virtual bool AllowTasksCompletion() const
 	{
-		static FStateTreeDetailsCustomizationSettings Settings;
-		return Settings;
-	}
-	
+		return true;
+	}	
 #endif // WITH_EDITOR
 };
 
