AnimNext: Pass-by-ref runtime refactor

diff --git a/Engine/Plugins/Experimental/AnimNextPoseSearch/Source/AnimNextPoseSearch/Private/PoseSearchResultEmulatorTrait.cpp b/Engine/Plugins/Experimental/AnimNextPoseSearch/Source/AnimNextPoseSearch/Private/PoseSearchResultEmulatorTrait.cpp
--- a/Engine/Plugins/Experimental/AnimNextPoseSearch/Source/AnimNextPoseSearch/Private/PoseSearchResultEmulatorTrait.cpp
+++ b/Engine/Plugins/Experimental/AnimNextPoseSearch/Source/AnimNextPoseSearch/Private/PoseSearchResultEmulatorTrait.cpp
@@ -4,16 +4,16 @@
 #include "Animation/AnimRootMotionProvider.h"
 #include "Animation/AnimSequence.h"
 #include "Animation/BlendSpace.h"
-#include "DataInterface/DataInterfaceStructAdapter.h"
 #include "Graph/AnimNextGraphInstance.h"
-#include "GraphInterfaces/AnimNextNativeDataInterface_AnimSequencePlayer.h"
-#include "GraphInterfaces/AnimNextNativeDataInterface_BlendSpacePlayer.h"
+#include "Traits/SequencePlayerTraitData.h"
+#include "Traits/BlendSpacePlayerTraitData.h"
 #include "PoseSearch/MultiAnimAsset.h"
 #include "TraitInterfaces/IBlendStack.h"
 #include "TraitInterfaces/IGraphFactory.h"
 #include "TraitInterfaces/ITimeline.h"
 #include "EvaluationVM/EvaluationVM.h"
-#include "GraphInterfaces/AnimNextNativeDataInterface_SynchronizeUsingGroups.h"
+#include "Factory/AnimGraphFactory.h"
+#include "Traits/SynchronizeUsingGroupsTraitData.h"
 
 namespace UE::AnimNext
 {
@@ -107,75 +107,66 @@ namespace UE::AnimNext
 								// If it's not a continuing pose search, then our search requested a new pose.
 								|| !Result.bIsContinuingPoseSearch;
 
-			if (bDoBlend)
+			if (bDoBlend && AnimationAsset)
 			{
-				UE::AnimNext::IBlendStack::FGraphRequest NewGraphRequest;
-				NewGraphRequest.BlendArgs = SharedData->GetBlendArguments(Binding);
-
-				// Create a default 'native data interface' payload via which we expect to communicate with the graph 
-				FAnimNextDataInterfacePayload GraphPayload;
-				FInstancedStruct AssetPlayerPayload;
+				FAnimNextFactoryParams FactoryParams = FAnimGraphFactory::GetDefaultParamsForClass(AnimationAsset->GetClass());
 				if (bIsBlendSpace)
 				{
+					FBlendSpacePlayerData BlendSpacePlayer;
+					BlendSpacePlayer.BlendSpace = Cast<UBlendSpace>(AnimationAsset);
+					BlendSpacePlayer.XAxisSamplePoint = XAxisSamplePoint;
+					BlendSpacePlayer.YAxisSamplePoint = YAxisSamplePoint;
+					BlendSpacePlayer.PlayRate = WantedPlayRate;
+					BlendSpacePlayer.StartPosition =  SelectedTime;
+					BlendSpacePlayer.bLoop = bLoop;
+					FactoryParams.PushPublicTrait(BlendSpacePlayer);
+
+					FSynchronizeUsingGroupsData SynchronizeUsingGroups;
 					// Tell blendspace samples to synchronize among themselves, but not with anything else.
-					FInstancedStruct SyncGroupPayload;
-					auto& NativeInterface = SyncGroupPayload.InitializeAs<FAnimNextNativeDataInterface_SynchronizeUsingGroups>();
-					NativeInterface.GroupName = NAME_None;
-					NativeInterface.GroupRole = EAnimGroupSynchronizationRole::AlwaysFollower;
-					NativeInterface.SyncMode = EAnimGroupSynchronizationMode::SynchronizeUsingUniqueGroupName;
-					NativeInterface.MatchSyncPoint = true;
-					GraphPayload.AddNative(MoveTemp(SyncGroupPayload));
-
-					FAnimNextNativeDataInterface_BlendSpacePlayer& NewNativeInterface = AssetPlayerPayload.InitializeAs<FAnimNextNativeDataInterface_BlendSpacePlayer>();
-					NewNativeInterface.BlendSpace = Cast<UBlendSpace>(AnimationAsset);
-
-					NewNativeInterface.XAxisSamplePoint = XAxisSamplePoint;
-					NewNativeInterface.YAxisSamplePoint = YAxisSamplePoint;
-
-					NewNativeInterface.PlayRate = WantedPlayRate;
-					NewNativeInterface.StartPosition =  SelectedTime;
-					NewNativeInterface.Loop = bLoop;
+					SynchronizeUsingGroups.GroupName = NAME_None;
+					SynchronizeUsingGroups.GroupRole = EAnimGroupSynchronizationRole::AlwaysFollower;
+					SynchronizeUsingGroups.SyncMode = EAnimGroupSynchronizationMode::SynchronizeUsingUniqueGroupName;
+					SynchronizeUsingGroups.bMatchSyncPoint = true;
+					FactoryParams.PushPublicTrait(SynchronizeUsingGroups);
 				}
 				else
 				{
-					FAnimNextNativeDataInterface_AnimSequencePlayer& NewNativeInterface = AssetPlayerPayload.InitializeAs<FAnimNextNativeDataInterface_AnimSequencePlayer>();
-					NewNativeInterface.AnimSequence = Cast<UAnimSequence>(AnimationAsset);
-					NewNativeInterface.PlayRate = WantedPlayRate;
-					NewNativeInterface.StartPosition =  SelectedTime;
-					NewNativeInterface.Loop = bLoop;
+					FSequencePlayerData AnimSequencePlayer;
+					AnimSequencePlayer.AnimSequence = Cast<UAnimSequence>(AnimationAsset);
+					AnimSequencePlayer.PlayRate = WantedPlayRate;
+					AnimSequencePlayer.StartPosition =  SelectedTime;
+					AnimSequencePlayer.bLoop = bLoop;
+					FactoryParams.PushPublicTrait(AnimSequencePlayer);
 				}
 
-				GraphPayload.AddNative(MoveTemp(AssetPlayerPayload));
-
-				const UAnimNextAnimationGraph* AnimationGraph = IGraphFactory::GetGraphFromObjectWithFallback(Context, Binding, AnimationAsset, GraphPayload);
+				const UAnimNextAnimationGraph* AnimationGraph = IGraphFactory::GetOrBuildGraph(Context, Binding, AnimationAsset, FactoryParams);
 				if (AnimationGraph != nullptr)
 				{
-					NewGraphRequest.GraphPayload = MoveTemp(GraphPayload);
+					IBlendStack::FGraphRequest NewGraphRequest;
+					NewGraphRequest.BlendArgs = SharedData->GetBlendArguments(Binding);
 					NewGraphRequest.AnimationGraph = AnimationGraph;
 					NewGraphRequest.FactoryObject = AnimationAsset;
+					NewGraphRequest.FactoryParams = MoveTemp(FactoryParams);
 					BlendStackBinding.PushGraph(Context, MoveTemp(NewGraphRequest));
 				}
 			}
-			
-			// HACK: Grab the active request again. It seems the previous active graph request can become invalid here.
-			BlendStackBinding.GetActiveGraph(Context, ActiveGraphRequest);
-			if (ActiveGraphRequest)
+
+			// Update the active graph instance according to continuing parameters
+			FAnimNextGraphInstance* ActiveGraphInstance = nullptr;
+			BlendStackBinding.GetActiveGraphInstance(Context, ActiveGraphInstance);
+			if (ActiveGraphInstance)
 			{
-				for (FStructView& Payload : ActiveGraphRequest->GraphPayload.Get())
+				ActiveGraphInstance->AccessVariablesStruct<FSequencePlayerData>([WantedPlayRate](FSequencePlayerData& InAnimSequencePlayer)
 				{
-					if (FAnimNextNativeDataInterface_AnimSequencePlayer* PlayAnimPayload = Payload.GetPtr<FAnimNextNativeDataInterface_AnimSequencePlayer>())
-					{
-						PlayAnimPayload->PlayRate = WantedPlayRate;
-						break;
-					}
-					else if (FAnimNextNativeDataInterface_BlendSpacePlayer* BlendSpacePayload = Payload.GetPtr<FAnimNextNativeDataInterface_BlendSpacePlayer>())
-					{
-						BlendSpacePayload->PlayRate = WantedPlayRate;
-						BlendSpacePayload->XAxisSamplePoint = XAxisSamplePoint;
-						BlendSpacePayload->YAxisSamplePoint = YAxisSamplePoint;
-						break;
-					}
-				}
+					InAnimSequencePlayer.PlayRate = WantedPlayRate;
+				});
+
+				ActiveGraphInstance->AccessVariablesStruct<FBlendSpacePlayerData>([XAxisSamplePoint, YAxisSamplePoint, WantedPlayRate](FBlendSpacePlayerData& InBlendSpacePlayer)
+				{
+					InBlendSpacePlayer.PlayRate = WantedPlayRate;
+					InBlendSpacePlayer.XAxisSamplePoint = XAxisSamplePoint;
+					InBlendSpacePlayer.YAxisSamplePoint = YAxisSamplePoint;
+				});
 			}
 		}
 
diff --git a/Engine/Plugins/Experimental/AnimNextStateTree/Source/AnimNextStateTree/Internal/AnimNextStateTreeEditorOnlyTypes.h b/Engine/Plugins/Experimental/AnimNextStateTree/Source/AnimNextStateTree/Internal/AnimNextStateTreeEditorOnlyTypes.h
--- a/Engine/Plugins/Experimental/AnimNextStateTree/Source/AnimNextStateTree/Internal/AnimNextStateTreeEditorOnlyTypes.h
+++ b/Engine/Plugins/Experimental/AnimNextStateTree/Source/AnimNextStateTree/Internal/AnimNextStateTreeEditorOnlyTypes.h
@@ -7,7 +7,7 @@
 #include "CoreMinimal.h"
 #include "Compilation/AnimNextGetFunctionHeaderCompileContext.h"
 
-class UAnimNextDataInterface_EditorData;
+class UAnimNextSharedVariables_EditorData;
 class URigVMGraph;
 
 struct FRigVMCompileSettings;
@@ -18,7 +18,7 @@ struct FRigVMClient;
  */
 struct ANIMNEXTSTATETREE_API FAnimNextStateTreeProgrammaticFunctionHeaderParams
 {
-	FAnimNextStateTreeProgrammaticFunctionHeaderParams(UAnimNextDataInterface_EditorData* InEditorData, const FRigVMCompileSettings& InSettings, FRigVMClient& InRigVMClient, FAnimNextGetFunctionHeaderCompileContext& InOutCompileContext)
+	FAnimNextStateTreeProgrammaticFunctionHeaderParams(UAnimNextSharedVariables_EditorData* InEditorData, const FRigVMCompileSettings& InSettings, FRigVMClient& InRigVMClient, FAnimNextGetFunctionHeaderCompileContext& InOutCompileContext)
 		: EditorData(InEditorData)
 		, Settings(InSettings)
 		, RigVMClient(InRigVMClient)
@@ -26,7 +26,7 @@ struct ANIMNEXTSTATETREE_API FAnimNextStateTreeProgrammaticFunctionHeaderParams
 	{
 	}
 
-	UAnimNextDataInterface_EditorData* EditorData;
+	UAnimNextSharedVariables_EditorData* EditorData;
 	const FRigVMCompileSettings& Settings;
 	FRigVMClient& RigVMClient;
 	FAnimNextGetFunctionHeaderCompileContext& OutCompileContext;
diff --git a/Engine/Plugins/Experimental/AnimNextStateTree/Source/AnimNextStateTree/Private/AnimNextStateTreeContext.cpp b/Engine/Plugins/Experimental/AnimNextStateTree/Source/AnimNextStateTree/Private/AnimNextStateTreeContext.cpp
--- a/Engine/Plugins/Experimental/AnimNextStateTree/Source/AnimNextStateTree/Private/AnimNextStateTreeContext.cpp
+++ b/Engine/Plugins/Experimental/AnimNextStateTree/Source/AnimNextStateTree/Private/AnimNextStateTreeContext.cpp
@@ -1,31 +1,24 @@
 // Copyright Epic Games, Inc. All Rights Reserved.
 
 #include "AnimNextStateTreeContext.h"
-
-#include "AnimNextAnimGraphSettings.h"
+#include "Factory/AnimGraphFactory.h"
+#include "TraitCore/NodeInstance.h"
 #include "TraitCore/TraitStackBinding.h"
 #include "TraitInterfaces/IBlendStack.h"
-#include "TraitInterfaces/IGraphFactory.h"
-#include "TraitInterfaces/ISmoothBlend.h"
 #include "TraitInterfaces/ITimeline.h"
 
-bool FAnimNextStateTreeTraitContext::PushAssetOntoBlendStack(TNonNullPtr<UObject> InAsset, const FAlphaBlendArgs& InBlendArguments, FAnimNextDataInterfacePayload&& InPayload) const
+bool FAnimNextStateTreeTraitContext::PushAssetOntoBlendStack(TNonNullPtr<UObject> InAsset, const FAlphaBlendArgs& InBlendArguments, const FAnimNextFactoryParams& InFactoryParams) const
 {
 	using namespace UE::AnimNext;
 
-	if (Binding == nullptr || Context == nullptr)
-	{
-		return false;
-	}
-
 	TTraitBinding<IBlendStack> BlendStackBinding;
 	if(!Binding->GetStackInterface<IBlendStack>(BlendStackBinding))
 	{
 		return false;
 	}
 
-	// Get a graph to instantiate
-	const UAnimNextAnimationGraph* AnimationGraph = IGraphFactory::GetGraphFromObjectWithFallback(*Context, *Binding, InAsset, InPayload);
+	// Instantiate/get a graph from the params
+	const UAnimNextAnimationGraph* AnimationGraph = FAnimGraphFactory::GetOrBuildGraph(InAsset, InFactoryParams.GetBuilder());
 	if(AnimationGraph == nullptr)
 	{
 		return false;
@@ -35,7 +28,7 @@ bool FAnimNextStateTreeTraitContext::PushAssetOntoBlendStack(TNonNullPtr<UObject
 	NewGraphRequest.BlendArgs = InBlendArguments;
 	NewGraphRequest.FactoryObject = InAsset;
 	NewGraphRequest.AnimationGraph = AnimationGraph;
-	NewGraphRequest.GraphPayload = MoveTemp(InPayload);
+	NewGraphRequest.FactoryParams = InFactoryParams;
 
 	BlendStackBinding.PushGraph(*Context, MoveTemp(NewGraphRequest));
 
@@ -44,14 +37,16 @@ bool FAnimNextStateTreeTraitContext::PushAssetOntoBlendStack(TNonNullPtr<UObject
 
 float FAnimNextStateTreeTraitContext::QueryPlaybackRatio(TNonNullPtr<UObject> InAsset) const
 {
-	if (Binding && Context)
+	UE::AnimNext::TTraitBinding<UE::AnimNext::ITimeline> Timeline;
+	if (Binding->GetStackInterface<UE::AnimNext::ITimeline>(Timeline))
 	{
-		UE::AnimNext::TTraitBinding<UE::AnimNext::ITimeline> Timeline;
-		if (Binding->GetStackInterface<UE::AnimNext::ITimeline>(Timeline))
-		{
-			return Timeline.GetState(*Context).GetPositionRatio();
-		}
+		return Timeline.GetState(*Context).GetPositionRatio();
 	}
 
 	return 1.0f;
 }
+
+FAnimNextGraphInstance& FAnimNextStateTreeTraitContext::GetGraphInstance() const
+{
+	return Binding->GetTraitPtr().GetNodeInstance()->GetOwner();
+}
\ No newline at end of file
diff --git a/Engine/Plugins/Experimental/AnimNextStateTree/Source/AnimNextStateTree/Private/AnimNextStateTreeContext.h b/Engine/Plugins/Experimental/AnimNextStateTree/Source/AnimNextStateTree/Private/AnimNextStateTreeContext.h
new file mode 100644
--- /dev/null
+++ b/Engine/Plugins/Experimental/AnimNextStateTree/Source/AnimNextStateTree/Private/AnimNextStateTreeContext.h
@@ -0,0 +1,56 @@
+// Copyright Epic Games, Inc. All Rights Reserved.
+
+#pragma once
+
+#include "AnimStateTreeTrait.h"
+#include "StructUtils/InstancedStruct.h"
+
+#include "AnimNextStateTreeContext.generated.h"
+
+namespace UE::AnimNext
+{
+struct FTraitBinding;
+struct FExecutionContext;
+struct FStateTreeTrait;
+}
+
+class UAnimNextAnimationGraph;
+struct FAlphaBlendArgs;
+struct FAnimNextFactoryParams;
+
+USTRUCT()
+struct FAnimNextStateTreeTraitContext
+{	
+	GENERATED_BODY()
+	
+	friend UE::AnimNext::FStateTreeTrait;
+	friend UE::AnimNext::FStateTreeTrait::FInstanceData;
+
+	FAnimNextStateTreeTraitContext() = default;
+
+	FAnimNextStateTreeTraitContext(UE::AnimNext::FExecutionContext& InContext, const UE::AnimNext::FTraitBinding& InBinding)
+		: Context(&InContext)
+		, Binding(&InBinding)
+	{}
+
+	bool PushAssetOntoBlendStack(TNonNullPtr<UObject> InAsset, const FAlphaBlendArgs& InBlendArguments, const FAnimNextFactoryParams& InFactoryParams) const;
+	float QueryPlaybackRatio(TNonNullPtr<UObject> InAsset) const;
+	FAnimNextGraphInstance& GetGraphInstance() const;
+
+	UE::AnimNext::FExecutionContext& GetContext() const
+	{
+		check(Context != nullptr);
+		return *Context;
+	}
+
+	const UE::AnimNext::FTraitBinding& GetBinding() const
+	{
+		check(Binding != nullptr);
+		return *Binding;
+	}
+
+private:
+	UE::AnimNext::FExecutionContext* Context = nullptr;
+	const UE::AnimNext::FTraitBinding* Binding = nullptr;
+};
+
diff --git a/Engine/Plugins/Experimental/AnimNextStateTree/Source/AnimNextStateTree/Private/AnimStateTreeTrait.cpp b/Engine/Plugins/Experimental/AnimNextStateTree/Source/AnimNextStateTree/Private/AnimStateTreeTrait.cpp
--- a/Engine/Plugins/Experimental/AnimNextStateTree/Source/AnimNextStateTree/Private/AnimStateTreeTrait.cpp
+++ b/Engine/Plugins/Experimental/AnimNextStateTree/Source/AnimNextStateTree/Private/AnimStateTreeTrait.cpp
@@ -5,6 +5,7 @@
 #include "AnimNextStateTreeContext.h"
 #include "AnimNextStateTreeSchema.h"
 #include "StateTreeExecutionContext.h"
+#include "UAFRigVMComponent.h"
 #include "Graph/AnimNextGraphInstance.h"
 #include "Logging/LogScopedVerbosityOverride.h"
 #include "TraitCore/NodeInstance.h"
@@ -99,110 +100,98 @@ namespace UE::AnimNext
 
 			FStateTreeExecutionContext StateTreeExecutionContext(*Owner, *InstanceData->StateTree, InstanceData->InstanceData);
 
-			FAnimNextStateTreeTraitContext TraitContext(Context, &Binding);
+			FAnimNextStateTreeTraitContext TraitContext(Context, Binding);
 			StateTreeExecutionContext.SetContextDataByName(UStateTreeAnimNextSchema::AnimStateTreeExecutionContextName, FStateTreeDataView(FAnimNextStateTreeTraitContext::StaticStruct(), reinterpret_cast<uint8*>(&TraitContext)));
 
-			FAnimNextGraphInstance* OwnerGraphInstance = &Binding.GetTraitPtr().GetNodeInstance()->GetOwner();
-			const bool bPropertyLayoutMatches = InstanceData->StateTree->GetDefaultParameters().GetPropertyBagStruct() == OwnerGraphInstance->GetVariables().GetPropertyBagStruct();
+			FAnimNextGraphInstance& OwnerGraphInstance = Binding.GetTraitPtr().GetNodeInstance()->GetOwner();
+#if DO_CHECK
+			check(OwnerGraphInstance.LayoutMatches(InstanceData->StateTree->GetDefaultParameters()));
+#endif
+			const FInstancedPropertyBag& StateTreeParameters = InstanceData->StateTree->GetDefaultParameters();
+			InstanceData->StateTreeExternalParameters.Reset();
+			
+			TArray<int32> StartPropertyOffsets;
+			FUAFRigVMComponent& RigVMComponent = OwnerGraphInstance.GetComponent<FUAFRigVMComponent>();
+			const FRigVMExtendedExecuteContext& ExtendedExecuteContext = RigVMComponent.GetExtendedExecuteContext();
 
-			if(OwnerGraphInstance && ensure(bPropertyLayoutMatches))
+			const int32 NumVariables = StateTreeParameters.GetNumPropertiesInBag();
+			for(int32 VariableIndex = 0; VariableIndex < NumVariables; ++VariableIndex)
 			{
-				const FInstancedPropertyBag& StateTreeParameters = InstanceData->StateTree->GetDefaultParameters();
-				FStructView MutableVariables = OwnerGraphInstance->GetMutableVariables().GetMutableValue();
-				InstanceData->StateTreeExternalParameters.Reset();
-				
-				TArray<int32> StartPropertyOffsets;
-				const FRigVMExtendedExecuteContext& ExtendedExecuteContext = OwnerGraphInstance->GetExtendedExecuteContext();
-
-				const int32 NumVariables = StateTreeParameters.GetNumPropertiesInBag();
-				for(int32 VariableIndex = 0; VariableIndex < NumVariables; ++VariableIndex)
-				{
-					if (ExtendedExecuteContext.ExternalVariableRuntimeData.IsValidIndex(VariableIndex) && MutableVariables.GetMemory())
-					{
-						const FPropertyBagPropertyDesc& Desc = StateTreeParameters.GetPropertyBagStruct()->GetPropertyDescs()[VariableIndex];
-						StartPropertyOffsets.Add(Desc.CachedProperty->GetOffset_ForInternal());
-					}
-				}
+				check(ExtendedExecuteContext.ExternalVariableRuntimeData.IsValidIndex(VariableIndex));
+				const FPropertyBagPropertyDesc& Desc = StateTreeParameters.GetPropertyBagStruct()->GetPropertyDescs()[VariableIndex];
+				StartPropertyOffsets.Add(Desc.CachedProperty->GetOffset_ForInternal());
+			}
 
-				for (const FPropertyBindingCopyInfoBatch& Batch : InstanceData->StateTree->GetPropertyBindings().GetCopyBatches())
+			for (const FPropertyBindingCopyInfoBatch& Batch : InstanceData->StateTree->GetPropertyBindings().GetCopyBatches())
+			{
+				for (const FPropertyBindingCopyInfo& Copy : InstanceData->StateTree->GetPropertyBindings().Super::GetBatchCopies(Batch))
 				{
-					for (const FPropertyBindingCopyInfo& Copy : InstanceData->StateTree->GetPropertyBindings().Super::GetBatchCopies(Batch))
+					const FStateTreeDataHandle& Handle = Copy.SourceDataHandle.Get<FStateTreeDataHandle>();
+					if (Handle.GetSource() == EStateTreeDataSourceType::ExternalGlobalParameterData)
 					{
-						const FStateTreeDataHandle& Handle = Copy.SourceDataHandle.Get<FStateTreeDataHandle>();
-						if (Handle.GetSource() == EStateTreeDataSourceType::ExternalGlobalParameterData)
-						{
 
-							const int32 RequiredOffset = [&Copy]()
-							{
-								if(Copy.SourceIndirection.Type == EPropertyBindingPropertyAccessType::Offset)
-								{
-									return (int32)Copy.SourceIndirection.Offset;
-								}
-
-								checkf(false, TEXT("Only expecting offset indirections for remapping"));
-								return Copy.SourceLeafProperty->GetOffset_ForInternal();
-							}();
-
-							uint8* MemoryPtr = nullptr;
-							const int32 NumOffsets = StartPropertyOffsets.Num();
-							for (int32 Index = 0; Index < NumOffsets; ++Index)
+						const int32 RequiredOffset = [&Copy]()
+						{
+							if(Copy.SourceIndirection.Type == EPropertyBindingPropertyAccessType::Offset)
 							{
-								const int32 NextOffset = (Index < (NumOffsets - 1)) ? StartPropertyOffsets[Index + 1] : INDEX_NONE;
-								if (RequiredOffset >= StartPropertyOffsets[Index] && (RequiredOffset < NextOffset || NextOffset == INDEX_NONE))
-								{
-									// Incorporate the UPropertyBag root-level property offset into the remapped memory-ptr itself, so that the property-access indirection works "as-normal"
-									MemoryPtr = (uint8*)ExtendedExecuteContext.ExternalVariableRuntimeData[Index].Memory - StartPropertyOffsets[Index];
-									break;
-								}
+								return (int32)Copy.SourceIndirection.Offset;
 							}
-							check(MemoryPtr != nullptr);
 
-							if (InstanceData->StateTreeExternalParameters.Add(Copy, MemoryPtr))
+							checkf(false, TEXT("Only expecting offset indirections for remapping"));
+							return Copy.SourceLeafProperty->GetOffset_ForInternal();
+						}();
+
+						uint8* MemoryPtr = nullptr;
+						const int32 NumOffsets = StartPropertyOffsets.Num();
+						for (int32 Index = 0; Index < NumOffsets; ++Index)
+						{
+							const int32 NextOffset = (Index < (NumOffsets - 1)) ? StartPropertyOffsets[Index + 1] : INDEX_NONE;
+							if (RequiredOffset >= StartPropertyOffsets[Index] && (RequiredOffset < NextOffset || NextOffset == INDEX_NONE))
 							{
-#if ENABLE_ANIM_DEBUG 
-								UE_CLOG(CVarLogPropertyBindingMemoryPtrInfo.GetValueOnAnyThread(), LogAnimation, Warning, TEXT("Mapped: Source: %s\nTarget: %sSize: %i\nOffset: %i\nType: %s"), *Copy.SourceLeafProperty->GetName(), *Copy.TargetLeafProperty->GetName(), Copy.CopySize, RequiredOffset, *FindObject<UEnum>(nullptr, TEXT("/Script/PropertyBindingUtils.EPropertyCopyType"))->GetNameStringByValue(static_cast<int64>(Copy.Type)));
+								// Incorporate the UPropertyBag root-level property offset into the remapped memory-ptr itself, so that the property-access indirection works "as-normal"
+								MemoryPtr = (uint8*)ExtendedExecuteContext.ExternalVariableRuntimeData[Index].Memory - StartPropertyOffsets[Index];
+								break;
 							}
-							else
-							{
-								UE_CLOG(CVarLogPropertyBindingMemoryPtrInfo.GetValueOnAnyThread(), LogAnimation, Warning, TEXT("Skipped: Source: %s\nTarget: %sSize: %i\nOffset: %i\nType: %s"), *Copy.SourceLeafProperty->GetName(), *Copy.TargetLeafProperty->GetName(), Copy.CopySize, RequiredOffset, *FindObject<UEnum>(nullptr, TEXT("/Script/PropertyBindingUtils.EPropertyCopyType"))->GetNameStringByValue(static_cast<int64>(Copy.Type)));
+						}
+						check(MemoryPtr != nullptr);
+
+						if (InstanceData->StateTreeExternalParameters.Add(Copy, MemoryPtr))
+						{
+#if ENABLE_ANIM_DEBUG 
+							UE_CLOG(CVarLogPropertyBindingMemoryPtrInfo.GetValueOnAnyThread(), LogAnimation, Warning, TEXT("Mapped: Source: %s\nTarget: %sSize: %i\nOffset: %i\nType: %s"), *Copy.SourceLeafProperty->GetName(), *Copy.TargetLeafProperty->GetName(), Copy.CopySize, RequiredOffset, *FindObject<UEnum>(nullptr, TEXT("/Script/PropertyBindingUtils.EPropertyCopyType"))->GetNameStringByValue(static_cast<int64>(Copy.Type)));
+						}
+						else
+						{
+							UE_CLOG(CVarLogPropertyBindingMemoryPtrInfo.GetValueOnAnyThread(), LogAnimation, Warning, TEXT("Skipped: Source: %s\nTarget: %sSize: %i\nOffset: %i\nType: %s"), *Copy.SourceLeafProperty->GetName(), *Copy.TargetLeafProperty->GetName(), Copy.CopySize, RequiredOffset, *FindObject<UEnum>(nullptr, TEXT("/Script/PropertyBindingUtils.EPropertyCopyType"))->GetNameStringByValue(static_cast<int64>(Copy.Type)));
 #endif // ENABLE_ANIM_DEBUG 
-					        }
-				        }
+						}
 					}
 				}
-				StateTreeExecutionContext.SetExternalGlobalParameters(&InstanceData->StateTreeExternalParameters);
-				
-				StateTreeExecutionContext.SetCollectExternalDataCallback(FOnCollectStateTreeExternalData::CreateLambda([this, &TraitContext]( const FStateTreeExecutionContext& Context, const UStateTree* StateTree, TArrayView<const FStateTreeExternalDataDesc> ExternalDataDescs, TArrayView<FStateTreeDataView> OutDataViews) -> bool
+			}
+			StateTreeExecutionContext.SetExternalGlobalParameters(&InstanceData->StateTreeExternalParameters);
+			
+			StateTreeExecutionContext.SetCollectExternalDataCallback(FOnCollectStateTreeExternalData::CreateLambda([this, &TraitContext]( const FStateTreeExecutionContext& Context, const UStateTree* StateTree, TArrayView<const FStateTreeExternalDataDesc> ExternalDataDescs, TArrayView<FStateTreeDataView> OutDataViews) -> bool
+			{
+				for (int32 Index = 0; Index < ExternalDataDescs.Num(); Index++)
 				{
-					for (int32 Index = 0; Index < ExternalDataDescs.Num(); Index++)
+					const FStateTreeExternalDataDesc& ItemDesc = ExternalDataDescs[Index];
+					if (ItemDesc.Struct != nullptr)
 					{
-						const FStateTreeExternalDataDesc& ItemDesc = ExternalDataDescs[Index];
-						if (ItemDesc.Struct != nullptr)
+						if (ItemDesc.Struct->IsChildOf(FAnimNextStateTreeTraitContext::StaticStruct()))
 						{
-							if (ItemDesc.Struct->IsChildOf(FAnimNextStateTreeTraitContext::StaticStruct()))
-							{
-								OutDataViews[Index] = FStateTreeDataView(FAnimNextStateTreeTraitContext::StaticStruct(), reinterpret_cast<uint8*>(&TraitContext));
-							}
+							OutDataViews[Index] = FStateTreeDataView(FAnimNextStateTreeTraitContext::StaticStruct(), reinterpret_cast<uint8*>(&TraitContext));
 						}
 					}
-
-					return true;
-				}));
-		
-				if (StateTreeExecutionContext.IsValid())
-				{
-					StateTreeExecutionContext.Start();
 				}
-			}
-			else
+
+				return true;
+			}));
+
+			if (StateTreeExecutionContext.IsValid())
 			{
-				if (StateTreeExecutionContext.IsValid())
-				{
-					UE_LOG(LogAnimation, Error, TEXT("Failed to start StateTree. PropertyLayoutMatch: %u"), bPropertyLayoutMatches);
-					StateTreeExecutionContext.Stop(EStateTreeRunStatus::Failed);
-				}
+				StateTreeExecutionContext.Start();
 			}
-		}	
+		}
 	}
 
 	void FStateTreeTrait::PreUpdate(FUpdateTraversalContext& Context, const TTraitBinding<IUpdate>& Binding, const FTraitUpdateState& TraitState) const
@@ -226,7 +215,7 @@ namespace UE::AnimNext
 			FAnimNextGraphInstance* OwnerGraphInstance = &Binding.GetTraitPtr().GetNodeInstance()->GetOwner();
 			if(OwnerGraphInstance && StateTreeExecutionContext.GetLastTickStatus() != EStateTreeRunStatus::Failed)
 			{
-				FAnimNextStateTreeTraitContext TraitContext(Context, &Binding);
+				FAnimNextStateTreeTraitContext TraitContext(Context, Binding);
 				StateTreeExecutionContext.SetContextDataByName(UStateTreeAnimNextSchema::AnimStateTreeExecutionContextName, FStateTreeDataView(FAnimNextStateTreeTraitContext::StaticStruct(), reinterpret_cast<uint8*>(&TraitContext)));
 				StateTreeExecutionContext.SetExternalGlobalParameters(&InstanceData->StateTreeExternalParameters);
 
diff --git a/Engine/Plugins/Experimental/AnimNextStateTree/Source/AnimNextStateTree/Private/Conditions/AnimNextStateTreeRigVMConditionBase.cpp b/Engine/Plugins/Experimental/AnimNextStateTree/Source/AnimNextStateTree/Private/Conditions/AnimNextStateTreeRigVMConditionBase.cpp
--- a/Engine/Plugins/Experimental/AnimNextStateTree/Source/AnimNextStateTree/Private/Conditions/AnimNextStateTreeRigVMConditionBase.cpp
+++ b/Engine/Plugins/Experimental/AnimNextStateTree/Source/AnimNextStateTree/Private/Conditions/AnimNextStateTreeRigVMConditionBase.cpp
@@ -8,6 +8,7 @@
 #include "StateTreeExecutionContext.h"
 #include "StructUtils/PropertyBag.h"
 #include "TraitCore/ExecutionContext.h"
+#include "AnimNextStateTreeContext.h"
 
 #if WITH_EDITOR
 #include "AnimNextStateTreeEditorOnlyTypes.h"
@@ -20,6 +21,8 @@
 
 #define LOCTEXT_NAMESPACE "AnimNextStateTreeRigVMConditionBase"
 
+#include "UAFRigVMComponent.h"
+
 #include UE_INLINE_GENERATED_CPP_BY_NAME(AnimNextStateTreeRigVMConditionBase)
 
 ////////////////////////////////////////////////////////////////////////////////////////////////////////////////////////
@@ -36,49 +39,36 @@ bool FAnimNextStateTreeRigVMConditionBase::TestCondition(FStateTreeExecutionCont
 	if (FAnimNextStateTreeRigVMConditionInstanceData* InstanceData = Context.GetInstanceDataPtr<FAnimNextStateTreeRigVMConditionInstanceData>(*this))
 	{
 		FAnimNextStateTreeTraitContext& ExecContext = Context.GetExternalData(TraitContextHandle);
-		const UE::AnimNext::FExecutionContext* AnimExecContext = ExecContext.GetAnimExecuteContext();
+		FAnimNextGraphInstance& GraphInstance = ExecContext.GetGraphInstance();
 
 		// Need const cast as VM execution isn't const.
-		if (TObjectPtr<UAnimNextDataInterface> RigVMInstance = const_cast<UAnimNextDataInterface*>(AnimExecContext->GetRootGraphInstance().GetDataInterface()))
+		if (UAnimNextRigVMAsset* RigVMInstance = const_cast<UAnimNextRigVMAsset*>(GraphInstance.GetAsset<UAnimNextRigVMAsset>()))
 		{
-			// Cache RigVM variable indexes, which can vary as they are parent derived
-			if (!InstanceData->VariableIndexCache.IsIndexCacheInitialized() && !InstanceData->VariableIndexCache.TryPopulateIndexCache(InstanceData->ParamData, RigVMInstance, RigVMFunctionHeader.Name, InternalResultName))
+#if WITH_EDITOR
+			if (!InstanceData->FunctionHandle.IsValidForVM(RigVMInstance->GetRigVM()->GetVMHash()))
+#else
+			if (!InstanceData->FunctionHandle.IsValid())
+#endif
 			{
-				return false;
+				InstanceData->FunctionHandle = RigVMInstance->GetFunctionHandle(InternalEventName);
 			}
 
-			// Prior to execution, copy over ST binding values to RigVM variables
-			FAnimNextGraphInstance& GraphInstance = AnimExecContext->GetRootGraphInstance();
-			TArrayView<const uint8> ArgumentIndexes = InstanceData->VariableIndexCache.GetVMArgumentIndexes();
-			TConstArrayView<FPropertyBagPropertyDesc> StateTreePropertyDescs = InstanceData->ParamData.GetPropertyBagStruct()->GetPropertyDescs();
-			TConstArrayView<FPropertyBagPropertyDesc> AnimNextGraphPropertyDescs = GraphInstance.GetVariables().GetPropertyBagStruct()->GetPropertyDescs();
-			FStructView StateTreeParamData = InstanceData->ParamData.GetMutableValue();
-
-			for (int32 Index = 0; Index < ArgumentIndexes.Num(); Index++)
+			if (!ensure(InstanceData->FunctionHandle.IsValid() && ResultIndex != INDEX_NONE))
 			{
-				uint8 RigVMExternalVariableIndex = ArgumentIndexes[Index];
-
-				const FPropertyBagPropertyDesc& StateTreePropertyDesc = StateTreePropertyDescs[Index];
-				void* SourceAddress = StateTreePropertyDesc.CachedProperty->ContainerPtrToValuePtr<void*>(StateTreeParamData.GetMemory());
-
-				const FPropertyBagPropertyDesc& AnimNextGraphDesc = AnimNextGraphPropertyDescs[RigVMExternalVariableIndex];
-				GraphInstance.GetMutableVariables().SetValue(AnimNextGraphDesc.Name, StateTreePropertyDesc.CachedProperty, StateTreeParamData.GetMemory());
+				return false;
 			}
 
-			FRigVMExtendedExecuteContext& ExtendedExecuteContext = GraphInstance.GetExtendedExecuteContext();
-			FAnimNextExecuteContext& AnimNextContext = GraphInstance.GetExtendedExecuteContext().GetPublicDataSafe<FAnimNextExecuteContext>();
+			FUAFRigVMComponent& RigVMComponent = GraphInstance.GetComponent<FUAFRigVMComponent>();
+			FRigVMExtendedExecuteContext& ExtendedExecuteContext = RigVMComponent.GetExtendedExecuteContext();
+			FAnimNextExecuteContext& AnimNextContext = ExtendedExecuteContext.GetPublicDataSafe<FAnimNextExecuteContext>();
 
 			FAnimNextGraphContextData ContextData(GraphInstance.GetModuleInstance(), &GraphInstance);
 			UE::AnimNext::FScopedExecuteContextData ContextDataScope(AnimNextContext, ContextData);
 
-			bool bResult = false;
-			RigVMInstance->GetVM()->ExecuteVM(ExtendedExecuteContext, InternalEventName);
-			if (ensure(ExtendedExecuteContext.ExternalVariableRuntimeData.IsValidIndex(InstanceData->VariableIndexCache.GetVMResultIndex())))
-			{
-				bResult = (bool)*ExtendedExecuteContext.ExternalVariableRuntimeData[InstanceData->VariableIndexCache.GetVMResultIndex()].Memory;
-			}
+			RigVMInstance->CallFunctionHandle(InstanceData->FunctionHandle, ExtendedExecuteContext, InstanceData->Parameters);
 
-			return bResult;
+			// Pull the result out of the property bag
+			return InstanceData->Parameters.GetValueBool(InstanceData->Parameters.GetPropertyBagStruct()->GetPropertyDescs()[ResultIndex]).GetValue();
 		}
 	}
 
@@ -123,26 +113,31 @@ void FAnimNextStateTreeRigVMConditionBase::PostEditNodeChangeChainProperty(const
 			};
 
 			RigVMFunctionHeader = GetRigVMFunctionHeader(ConditionFunctionName);
-			InstanceData->ParamData = FRigVMMemoryStorageStruct();
+			InstanceData->Parameters.Reset();
+			ResultIndex = INDEX_NONE;
 
-			TArray<FRigVMPropertyDescription> RigVMPropertyDescriptions;
+			TArray<FPropertyBagPropertyDesc> PropertyDescs;
+			PropertyDescs.Reserve(RigVMFunctionHeader.Arguments.Num());
 			for (const FRigVMGraphFunctionArgument& Argument : RigVMFunctionHeader.Arguments)
 			{
-				if (Argument.Direction == ERigVMPinDirection::Input)
+				if (Argument.Direction == ERigVMPinDirection::Input || Argument.Direction == ERigVMPinDirection::Output)
 				{
 					FString CppTypeString = Argument.CPPType.ToString();
 					UObject* CppTypeObject = Argument.CPPTypeObject.Get();
-					FRigVMPropertyDescription RigVMPropertyDescription = FRigVMPropertyDescription(Argument.Name, CppTypeString, CppTypeObject, Argument.DefaultValue);
-					RigVMPropertyDescriptions.Add(RigVMPropertyDescription);
-				}
-				else if (Argument.Direction == ERigVMPinDirection::Output)
-				{
-					ensure(Argument.CPPType == RigVMTypeUtils::BoolTypeName);
-					ResultName = Argument.Name;
+					FPropertyBagPropertyDesc PropertyDesc = FRigVMMemoryStorageStruct::GeneratePropertyBagDescriptor(FRigVMPropertyDescription(Argument.Name, CppTypeString, CppTypeObject, Argument.DefaultValue));
+					if (Argument.Direction == ERigVMPinDirection::Output && Argument.CPPType == RigVMTypeUtils::BoolTypeName)
+					{
+						PropertyDesc.PropertyFlags &= ~CPF_Edit;	// Hide the result from the UI
+						ResultIndex = PropertyDescs.Num();
+					}
+					PropertyDescs.Add(PropertyDesc);
 				}
 			}
-		
-			InstanceData->ParamData.AddProperties(RigVMPropertyDescriptions);
+
+			// TODO: Constrain functions we can pick to those returning only a bool
+			ensure(ResultIndex != INDEX_NONE);
+
+			InstanceData->Parameters.AddProperties(PropertyDescs);
 		}
 	}
 }
@@ -155,13 +150,11 @@ void FAnimNextStateTreeRigVMConditionBase::GetProgrammaticFunctionHeaders(FAnimN
 	
 	StateName = State->Name;
 	NodeId = Desc.ID;
-	InternalResultName = FName(FUtils::MakeFunctionWrapperVariableName(RigVMFunctionHeader.Name, ResultName));
+	ResultIndex = INDEX_NONE;
 	InternalEventName = FName(FUtils::MakeFunctionWrapperEventName(RigVMFunctionHeader.Name));
 
 	FAnimNextProgrammaticFunctionHeader AnimNextFunctionHeader = {};
-	AnimNextFunctionHeader.Wrapped = RigVMFunctionHeader;
-	AnimNextFunctionHeader.bGenerateParamVariables = true;
-	AnimNextFunctionHeader.bGenerateReturnVariables = true;
+	AnimNextFunctionHeader.FunctionHeader = RigVMFunctionHeader;
 	OutCompileContext.GetMutableFunctionHeaders().Add(AnimNextFunctionHeader);
 }
 
diff --git a/Engine/Plugins/Experimental/AnimNextStateTree/Source/AnimNextStateTree/Private/Tasks/AnimNextStateTreeGraphInstanceTask.cpp b/Engine/Plugins/Experimental/AnimNextStateTree/Source/AnimNextStateTree/Private/Tasks/AnimNextStateTreeGraphInstanceTask.cpp
new file mode 100644
--- /dev/null
+++ b/Engine/Plugins/Experimental/AnimNextStateTree/Source/AnimNextStateTree/Private/Tasks/AnimNextStateTreeGraphInstanceTask.cpp
@@ -0,0 +1,161 @@
+// Copyright Epic Games, Inc. All Rights Reserved.
+
+#include "Tasks/AnimNextStateTreeGraphInstanceTask.h"
+
+#include "AnimNextAnimGraphSettings.h"
+#include "StateTreeLinker.h"
+#include "StateTreeExecutionContext.h"
+#include "Factory/AnimNextFactoryParams.h"
+#include "AnimNextStateTreeContext.h"
+#include "Factory/AnimGraphFactory.h"
+#include "UObject/FortniteMainBranchObjectVersion.h"
+
+#if WITH_EDITORONLY_DATA
+bool FAnimNextGraphInstanceTaskInstanceData::Serialize(FArchive& Ar)
+{
+	Ar.UsingCustomVersion(FFortniteMainBranchObjectVersion::GUID);
+	return false;
+}
+
+void FAnimNextGraphInstanceTaskInstanceData::PostSerialize(const FArchive& Ar)
+{
+	if (Ar.IsLoading() && Ar.CustomVer(FFortniteMainBranchObjectVersion::GUID) < FFortniteMainBranchObjectVersion::AnimNextVariableReferences)
+	{
+		// Prior to this we only supported 'imprecise' payloads so try to do our best be existing content, but some data loss is possible
+		if (Asset != nullptr && !Asset->IsA<UAnimNextAnimationGraph>())
+		{
+			Parameters = UE::AnimNext::FAnimGraphFactory::GetDefaultParamsForObject(Asset);
+
+			// Match any legacy parameter values by name
+			if(Payload_DEPRECATED.GetNumPropertiesInBag() > 0 && Parameters.Builder.Stacks.Num() > 0)
+			{
+				const UPropertyBag* PropertyBag = Payload_DEPRECATED.GetPropertyBagStruct();
+				const uint8* SourceContainerPtr = Payload_DEPRECATED.GetValue().GetMemory();
+
+				auto PatchInLegacyPropertyByName = [this, SourceContainerPtr](const FPropertyBagPropertyDesc& InDesc)
+				{
+					for (TInstancedStruct<FAnimNextTraitSharedData>& Struct : Parameters.Builder.Stacks[0].TraitStructs)
+					{
+						uint8* TargetContainerPtr = Struct.GetMutableMemory();
+						for (TFieldIterator<FProperty> It(Struct.GetScriptStruct()); It; ++It)
+						{
+							auto NameMatches = [](const FProperty* InProperty, FName InNameA, FName InNameB)
+							{
+								if (InNameA == InNameB)
+								{
+									return true;
+								}
+
+								// Match bool properties with no 'b' prefix
+								if (InProperty->IsA<FBoolProperty>())
+								{
+									FString StringA = InNameA.ToString();
+									StringA.RemoveFromStart(TEXT("b"));
+									FString StringB = InNameB.ToString();
+									StringB.RemoveFromStart(TEXT("b"));
+									return StringA.Equals(StringB, ESearchCase::IgnoreCase);
+								}
+								return false; 
+							};
+
+							if (NameMatches(*It, It->GetFName(), InDesc.Name) && It->GetClass() == InDesc.CachedProperty->GetClass())
+							{
+								const uint8* SourcePtr = InDesc.CachedProperty->ContainerPtrToValuePtr<uint8>(SourceContainerPtr);
+								uint8* TargetPtr = It->ContainerPtrToValuePtr<uint8>(TargetContainerPtr);
+								It->CopyCompleteValue(TargetPtr, SourcePtr);
+								break;
+							}
+						}
+					}
+				};
+
+				for (const FPropertyBagPropertyDesc& Desc : PropertyBag->GetPropertyDescs())
+				{
+					PatchInLegacyPropertyByName(Desc);
+				}
+			}
+		}
+	}
+}
+#endif
+
+FAnimNextStateTreeGraphInstanceTask::FAnimNextStateTreeGraphInstanceTask()
+{
+	// Re-selecting the same state should not cause a re-trigger of EnterState()
+	bShouldStateChangeOnReselect = false;
+}
+
+bool FAnimNextStateTreeGraphInstanceTask::Link(FStateTreeLinker& Linker)
+{
+	Linker.LinkExternalData(TraitContextHandle);	
+	return true;
+}
+
+EStateTreeRunStatus FAnimNextStateTreeGraphInstanceTask::EnterState(FStateTreeExecutionContext& Context, const FStateTreeTransitionResult& Transition) const
+{
+	FAnimNextStateTreeTraitContext& ExecContext = Context.GetExternalData(TraitContextHandle);
+	FInstanceDataType& InstanceData = Context.GetInstanceData(*this);
+
+	if (ExecContext.PushAssetOntoBlendStack(InstanceData.Asset.Get(), InstanceData.BlendOptions, InstanceData.Parameters))
+	{
+		return EStateTreeRunStatus::Running;
+	}
+	
+	return EStateTreeRunStatus::Failed;
+}
+
+EStateTreeRunStatus FAnimNextStateTreeGraphInstanceTask::Tick(FStateTreeExecutionContext& Context, const float DeltaTime) const
+{
+	FInstanceDataType& InstanceData = Context.GetInstanceData(*this);
+	if (InstanceData.bContinueTicking)
+	{
+		FAnimNextStateTreeTraitContext& ExecContext = Context.GetExternalData(TraitContextHandle);
+
+		InstanceData.PlaybackRatio = ExecContext.QueryPlaybackRatio(InstanceData.Asset.Get());
+
+		return EStateTreeRunStatus::Running;
+	}
+	
+	return EStateTreeRunStatus::Succeeded;
+}
+
+void FAnimNextStateTreeGraphInstanceTask::ExitState(FStateTreeExecutionContext& Context, const FStateTreeTransitionResult& Transition) const
+{
+	FAnimNextStateTreeTaskBase::ExitState(Context, Transition);
+}
+
+#if WITH_EDITOR
+void FAnimNextStateTreeGraphInstanceTask::PostEditInstanceDataChangeChainProperty(const FPropertyChangedChainEvent& PropertyChangedEvent, FStateTreeDataView InstanceDataView)
+{
+	const FProperty* Property = PropertyChangedEvent.Property;
+	if(Property == nullptr)
+	{
+		return;
+	}
+
+	if(Property->GetFName() == GET_MEMBER_NAME_CHECKED(FAnimNextGraphInstanceTaskInstanceData, Asset))
+	{
+		// Repopulate payload property bag if the selected asset changes
+		FAnimNextGraphInstanceTaskInstanceData& InstanceData = InstanceDataView.GetMutable<FAnimNextGraphInstanceTaskInstanceData>();
+
+		if (InstanceData.Asset == nullptr || InstanceData.Asset.IsA<UAnimNextAnimationGraph>())
+		{
+			InstanceData.Parameters.Reset();
+		}
+		else
+		{
+			// TODO: Copy over matching values before overwrite to avoid losing user edits to matching trait data
+			InstanceData.Parameters = UE::AnimNext::FAnimGraphFactory::GetDefaultParamsForObject(InstanceData.Asset);
+		}
+	}
+}
+
+void FAnimNextStateTreeGraphInstanceTask::GetObjectReferences(TArray<const UObject*>& OutReferencedObjects, FStateTreeDataView InstanceDataView) const
+{
+	if (InstanceDataView.IsValid())
+	{
+		OutReferencedObjects.Add(InstanceDataView.Get<FAnimNextGraphInstanceTaskInstanceData>().Asset.Get());
+	}	
+}
+#endif
+
diff --git a/Engine/Plugins/Experimental/AnimNextStateTree/Source/AnimNextStateTree/Private/Tasks/AnimNextStateTreeRigVMTaskBase.cpp b/Engine/Plugins/Experimental/AnimNextStateTree/Source/AnimNextStateTree/Private/Tasks/AnimNextStateTreeRigVMTaskBase.cpp
--- a/Engine/Plugins/Experimental/AnimNextStateTree/Source/AnimNextStateTree/Private/Tasks/AnimNextStateTreeRigVMTaskBase.cpp
+++ b/Engine/Plugins/Experimental/AnimNextStateTree/Source/AnimNextStateTree/Private/Tasks/AnimNextStateTreeRigVMTaskBase.cpp
@@ -8,6 +8,7 @@
 #include "StateTreeExecutionContext.h"
 #include "StructUtils/PropertyBag.h"
 #include "TraitCore/ExecutionContext.h"
+#include "AnimNextStateTreeContext.h"
 
 #if WITH_EDITOR
 #include "AnimNextStateTreeEditorOnlyTypes.h"
@@ -20,6 +21,9 @@
 
 #define LOCTEXT_NAMESPACE "AnimNextStateTreeRigVMTaskBase"
 
+#include "AnimNextRigVMAsset.h"
+#include "UAFRigVMComponent.h"
+
 #include UE_INLINE_GENERATED_CPP_BY_NAME(AnimNextStateTreeRigVMTaskBase)
 
 ////////////////////////////////////////////////////////////////////////////////////////////////////////////////////////
@@ -36,51 +40,38 @@ EStateTreeRunStatus FAnimNextStateTreeRigVMTaskBase::EnterState(FStateTreeExecut
 	if (FAnimNextStateTreeRigVMTaskInstanceData* InstanceData = Context.GetInstanceDataPtr<FAnimNextStateTreeRigVMTaskInstanceData>(*this))
 	{
 		FAnimNextStateTreeTraitContext& ExecContext = Context.GetExternalData(TraitContextHandle);
-		const UE::AnimNext::FExecutionContext* AnimExecContext = ExecContext.GetAnimExecuteContext();
+		FAnimNextGraphInstance& GraphInstance = ExecContext.GetGraphInstance();
 
 		// Need const cast as VM execution isn't const.
-		if (UAnimNextDataInterface* RigVMInstance = const_cast<UAnimNextDataInterface*>(AnimExecContext->GetRootGraphInstance().GetDataInterface()))
+		if (UAnimNextRigVMAsset* RigVMAsset = const_cast<UAnimNextRigVMAsset*>(GraphInstance.GetAsset<UAnimNextRigVMAsset>()))
 		{
-			// @TODO: For now don't pass result name. Since we use fire and forget
-			// Cache RigVM variable indexes, which can vary as they are parent derived
-			if (!InstanceData->VariableIndexCache.IsIndexCacheInitialized() && !InstanceData->VariableIndexCache.TryPopulateIndexCache(InstanceData->ParamData, RigVMInstance, RigVMFunctionHeader.Name, NAME_None))
+#if WITH_EDITOR
+			if (!InstanceData->FunctionHandle.IsValidForVM(RigVMAsset->GetRigVM()->GetVMHash()))
+#else
+			if (!InstanceData->FunctionHandle.IsValid())
+#endif
 			{
-				return EStateTreeRunStatus::Failed;
+				InstanceData->FunctionHandle = RigVMAsset->GetFunctionHandle(InternalEventName);
 			}
 
-			// Prior to execution, copy over ST binding values to RigVM variables
-			FAnimNextGraphInstance& GraphInstance = AnimExecContext->GetRootGraphInstance();
-			TArrayView<const uint8> ArgumentIndexes = InstanceData->VariableIndexCache.GetVMArgumentIndexes();
-			TConstArrayView<FPropertyBagPropertyDesc> StateTreePropertyDescs = InstanceData->ParamData.GetPropertyBagStruct()->GetPropertyDescs();
-			TConstArrayView<FPropertyBagPropertyDesc> AnimNextGraphPropertyDescs = GraphInstance.GetVariables().GetPropertyBagStruct()->GetPropertyDescs();
-			FStructView StateTreeParamData = InstanceData->ParamData.GetMutableValue();
-
-			for (int32 Index = 0; Index < ArgumentIndexes.Num(); Index++)
+			if (!ensure(InstanceData->FunctionHandle.IsValid() && ResultIndex != INDEX_NONE))
 			{
-				uint8 RigVMExternalVariableIndex = ArgumentIndexes[Index];
-
-				const FPropertyBagPropertyDesc& StateTreePropertyDesc = StateTreePropertyDescs[Index];
-				void* SourceAddress = StateTreePropertyDesc.CachedProperty->ContainerPtrToValuePtr<void*>(StateTreeParamData.GetMemory());
-
-				const FPropertyBagPropertyDesc& AnimNextGraphDesc = AnimNextGraphPropertyDescs[RigVMExternalVariableIndex];
-				GraphInstance.GetMutableVariables().SetValue(AnimNextGraphDesc.Name, StateTreePropertyDesc.CachedProperty, StateTreeParamData.GetMemory());
+				return EStateTreeRunStatus::Failed;
 			}
 
-			FRigVMExtendedExecuteContext& ExtendedExecuteContext = GraphInstance.GetExtendedExecuteContext();
-			FAnimNextExecuteContext& AnimNextContext = GraphInstance.GetExtendedExecuteContext().GetPublicDataSafe<FAnimNextExecuteContext>();
+			FUAFRigVMComponent& RigVMComponent = GraphInstance.GetComponent<FUAFRigVMComponent>();
+			FRigVMExtendedExecuteContext& ExtendedExecuteContext = RigVMComponent.GetExtendedExecuteContext();
+			FAnimNextExecuteContext& AnimNextContext = ExtendedExecuteContext.GetPublicDataSafe<FAnimNextExecuteContext>();
 
 			FAnimNextGraphContextData ContextData(GraphInstance.GetModuleInstance(), &GraphInstance);
 			UE::AnimNext::FScopedExecuteContextData ContextDataScope(AnimNextContext, ContextData);
 
 			// @TODO: Add option to execute per tick?
-			RigVMInstance->GetVM()->ExecuteVM(ExtendedExecuteContext, InternalEventName);
+			RigVMAsset->CallFunctionHandle(InstanceData->FunctionHandle, ExtendedExecuteContext, InstanceData->Parameters);
 
 			// @TODO: For now we use a fire & forget model. But in the future we might want users to be able to bind to the result value from a function
-			// This will require a special UI customization, & standalone execute is useful for modifying existing vars. So save for later effort.
-			//if (ensure(ExtendedExecuteContext.ExternalVariableRuntimeData.IsValidIndex(InstanceData->VariableIndexCache.GetVMResultIndex())))
-			//{
-			//	auto Result = *ExtendedExecuteContext.ExternalVariableRuntimeData[InstanceData->VariableIndexCache.GetVMResultIndex()].Memory;
-			//}
+			// This will require a special UI customization, & standalone execute is useful for modifying existing vars. So save for later effort.	
+			//return InstanceData->Parameters.GetValueBool(InstanceData->Parameters.GetPropertyBagStruct()->GetPropertyDescs()[ResultIndex]).GetValue();
 
 			return EStateTreeRunStatus::Running;
 		}
@@ -127,26 +118,28 @@ void FAnimNextStateTreeRigVMTaskBase::PostEditNodeChangeChainProperty(const FPro
 			};
 
 			RigVMFunctionHeader = GetRigVMFunctionHeader(TaskFunctionName);
-			InstanceData->ParamData = FRigVMMemoryStorageStruct();
+			InstanceData->Parameters.Reset();
+			ResultIndex = INDEX_NONE;
 
-			TArray<FRigVMPropertyDescription> RigVMPropertyDescriptions;
+			TArray<FPropertyBagPropertyDesc> PropertyDescs;
+			PropertyDescs.Reserve(RigVMFunctionHeader.Arguments.Num());
 			for (const FRigVMGraphFunctionArgument& Argument : RigVMFunctionHeader.Arguments)
 			{
-				if (Argument.Direction == ERigVMPinDirection::Input)
+				if (Argument.Direction == ERigVMPinDirection::Input || Argument.Direction == ERigVMPinDirection::Output)
 				{
 					FString CppTypeString = Argument.CPPType.ToString();
 					UObject* CppTypeObject = Argument.CPPTypeObject.Get();
-					FRigVMPropertyDescription RigVMPropertyDescription = FRigVMPropertyDescription(Argument.Name, CppTypeString, CppTypeObject, Argument.DefaultValue);
-					RigVMPropertyDescriptions.Add(RigVMPropertyDescription);
-				}
-				else if (Argument.Direction == ERigVMPinDirection::Output)
-				{
-					ensure(Argument.CPPType == RigVMTypeUtils::BoolTypeName);
-					ResultName = Argument.Name;
+					FPropertyBagPropertyDesc PropertyDesc = FRigVMMemoryStorageStruct::GeneratePropertyBagDescriptor(FRigVMPropertyDescription(Argument.Name, CppTypeString, CppTypeObject, Argument.DefaultValue));
+					if (Argument.Direction == ERigVMPinDirection::Output && Argument.CPPType == RigVMTypeUtils::BoolTypeName)
+					{
+						PropertyDesc.PropertyFlags &= ~CPF_Edit;	// Hide the result from the UI
+						ResultIndex = PropertyDescs.Num();
+					}
+					PropertyDescs.Add(PropertyDesc);
 				}
 			}
 		
-			InstanceData->ParamData.AddProperties(RigVMPropertyDescriptions);
+			InstanceData->Parameters.AddProperties(PropertyDescs);
 		}
 	}
 }
@@ -159,14 +152,11 @@ void FAnimNextStateTreeRigVMTaskBase::GetProgrammaticFunctionHeaders(FAnimNextSt
 	
 	StateName = State->Name;
 	NodeId = Desc.ID;
-	InternalResultName = FName(FUtils::MakeFunctionWrapperVariableName(RigVMFunctionHeader.Name, ResultName));
+	ResultIndex = INDEX_NONE;
 	InternalEventName = FName(FUtils::MakeFunctionWrapperEventName(RigVMFunctionHeader.Name));
 
 	FAnimNextProgrammaticFunctionHeader AnimNextFunctionHeader = {};
-	AnimNextFunctionHeader.Wrapped = RigVMFunctionHeader;
-	AnimNextFunctionHeader.bGenerateParamVariables = true;
-	// @TODO: For now don't gen return variables. Since we use fire and forget
-	AnimNextFunctionHeader.bGenerateReturnVariables = false;
+	AnimNextFunctionHeader.FunctionHeader = RigVMFunctionHeader;
 	OutCompileContext.GetMutableFunctionHeaders().Add(AnimNextFunctionHeader);
 }
 
diff --git a/Engine/Plugins/Experimental/AnimNextStateTree/Source/AnimNextStateTreeUncookedOnly/Internal/AnimNextStateTreeEditorData.cpp b/Engine/Plugins/Experimental/AnimNextStateTree/Source/AnimNextStateTreeUncookedOnly/Internal/AnimNextStateTreeEditorData.cpp
--- a/Engine/Plugins/Experimental/AnimNextStateTree/Source/AnimNextStateTreeUncookedOnly/Internal/AnimNextStateTreeEditorData.cpp
+++ b/Engine/Plugins/Experimental/AnimNextStateTree/Source/AnimNextStateTreeUncookedOnly/Internal/AnimNextStateTreeEditorData.cpp
@@ -66,3 +66,11 @@ void UAnimNextStateTreeTreeEditorData::HandleStateTreeAssetChanges(UAnimNextRigV
 		break;
 	}
 }
+
+const FInstancedPropertyBag& UAnimNextStateTreeTreeEditorData::GetRootParametersPropertyBag() const
+{
+	UStateTree* OuterStateTree = CastChecked<UStateTree>(GetOuter());
+	UAnimNextStateTree* OuterAnimNextStateTree = CastChecked<UAnimNextStateTree>(OuterStateTree->GetOuter());
+	UAnimNextStateTree_EditorData* AnimNextEditorData = UE::AnimNext::UncookedOnly::FUtils::GetEditorData<UAnimNextStateTree_EditorData>(OuterAnimNextStateTree);
+	return AnimNextEditorData->CombinedVariablesPropertyBag;
+}
diff --git a/Engine/Plugins/Experimental/AnimNextStateTree/Source/AnimNextStateTreeUncookedOnly/Internal/AnimNextStateTree_EditorData.cpp b/Engine/Plugins/Experimental/AnimNextStateTree/Source/AnimNextStateTreeUncookedOnly/Internal/AnimNextStateTree_EditorData.cpp
--- a/Engine/Plugins/Experimental/AnimNextStateTree/Source/AnimNextStateTreeUncookedOnly/Internal/AnimNextStateTree_EditorData.cpp
+++ b/Engine/Plugins/Experimental/AnimNextStateTree/Source/AnimNextStateTreeUncookedOnly/Internal/AnimNextStateTree_EditorData.cpp
@@ -4,11 +4,12 @@
 
 #include "AnimGraphUncookedOnlyUtils.h"
 #include "AnimNextStateTree.h"
+#include "AnimNextStateTreeEditorData.h"
 #include "AnimNextStateTreeEditorOnlyTypes.h"
 #include "AnimNextStateTreeWorkspaceAssetUserData.h"
 #include "Conditions/AnimNextStateTreeRigVMConditionBase.h"
 #include "Compilation/AnimNextGetGraphCompileContext.h"
-#include "DataInterface/AnimNextDataInterface_EditorData.h"
+#include "Variables/AnimNextSharedVariables_EditorData.h"
 #include "StateTree.h"
 #include "StateTreeCompilerLog.h"
 #include "StateTreeEditingSubsystem.h"
@@ -26,7 +27,7 @@
 #include "TraitCore/TraitRegistry.h"
 #include "TraitCore/Trait.h"
 #include "AnimStateTreeTrait.h"
-#include "Entries/AnimNextDataInterfaceEntry.h"
+#include "Entries/AnimNextSharedVariablesEntry.h"
 #include "Traits/BlendStackTrait.h"
 #include "Traits/BlendSmoother.h"
 
@@ -39,13 +40,20 @@ void UAnimNextStateTree_EditorData::RecompileVM()
 {
 	UAnimNextAnimationGraph_EditorData::RecompileVM();
 
+	// Refresh the 'dummy' property bag we have with all combined variables now we have compiled
+	RefreshCombinedVariablesPropertyBag();
+
+	// Skip further state tree compilation if we had errors in regular asset compilation. It can create spurious errors
+	if (bErrorsDuringCompilation)
+	{
+		return;
+	}
+
 	UAnimNextStateTree* AnimationStateTree = UE::AnimNext::UncookedOnly::FUtils::GetAsset<UAnimNextStateTree>(this);
 
 	UStateTree* InnerStateTree = AnimationStateTree->StateTree;
 
-	UStateTreeEditorData* InnerEditorData = Cast<UStateTreeEditorData>(InnerStateTree->EditorData);
-
-	// Recompile ST as we just updated our variables propertybag. It may have changed from an external data interface update.
+	// Recompile ST as we just updated our variables propertybag. It may have changed from an external shared variables update.
 	FStateTreeCompilerLog Log;
 	bool bCompileSucceeded = UStateTreeEditingSubsystem::CompileStateTree(InnerStateTree, Log);
 	ensureMsgf(bCompileSucceeded, TEXT("Failed to compile state tree after data update: %s"), *InnerStateTree->GetFName().ToString());
@@ -57,7 +65,7 @@ TConstArrayView<TSubclassOf<UAnimNextRigVMAssetEntry>> UAnimNextStateTree_Editor
 {
 		UAnimNextVariableEntry::StaticClass(),
 		UAnimNextAnimationGraphEntry::StaticClass(),
-		UAnimNextDataInterfaceEntry::StaticClass(),
+		UAnimNextSharedVariablesEntry::StaticClass(),
 	};
 
 	return Classes;
@@ -212,3 +220,8 @@ void UAnimNextStateTree_EditorData::OnPreCompileGetProgrammaticGraphs(const FRig
 		OutCompileContext.GetMutableProgrammaticGraphs().Add(Graph);
 	}
 }
+
+void UAnimNextStateTree_EditorData::RefreshCombinedVariablesPropertyBag()
+{
+	CombinedVariablesPropertyBag = GenerateCombinedPropertyBag();
+}
\ No newline at end of file
diff --git a/Engine/Plugins/Experimental/AnimNextStateTree/Source/AnimNextStateTreeUncookedOnly/Internal/AnimNextStateTree_EditorData.h b/Engine/Plugins/Experimental/AnimNextStateTree/Source/AnimNextStateTreeUncookedOnly/Internal/AnimNextStateTree_EditorData.h
--- a/Engine/Plugins/Experimental/AnimNextStateTree/Source/AnimNextStateTreeUncookedOnly/Internal/AnimNextStateTree_EditorData.h
+++ b/Engine/Plugins/Experimental/AnimNextStateTree/Source/AnimNextStateTreeUncookedOnly/Internal/AnimNextStateTree_EditorData.h
@@ -15,6 +15,7 @@ protected:
 	virtual TSubclassOf<UAssetUserData> GetAssetUserDataClass() const override;
 
 	friend class UAnimNextStateTreeFactory;
+	friend class UAnimNextStateTreeTreeEditorData;
 
 	// IRigVMClientHost interface
 	virtual void RecompileVM() override;
@@ -27,4 +28,11 @@ protected:
 
 	// Allows this asset to generate graphs to be injected at compilation time
 	virtual void OnPreCompileGetProgrammaticGraphs(const FRigVMCompileSettings& InSettings, FAnimNextGetGraphCompileContext& OutCompileContext) override;
+
+	void RefreshCombinedVariablesPropertyBag();
+
+private:
+	// Property bag that forms the root parameters. Consists of all combined variables from shared variable assets/structs.
+	UPROPERTY()
+	FInstancedPropertyBag CombinedVariablesPropertyBag;
 };
\ No newline at end of file
diff --git a/Engine/Plugins/Experimental/AnimNextStateTree/Source/AnimNextStateTreeUncookedOnly/Public/AnimNextStateTreeEditorData.h b/Engine/Plugins/Experimental/AnimNextStateTree/Source/AnimNextStateTreeUncookedOnly/Public/AnimNextStateTreeEditorData.h
--- a/Engine/Plugins/Experimental/AnimNextStateTree/Source/AnimNextStateTreeUncookedOnly/Public/AnimNextStateTreeEditorData.h
+++ b/Engine/Plugins/Experimental/AnimNextStateTree/Source/AnimNextStateTreeUncookedOnly/Public/AnimNextStateTreeEditorData.h
@@ -18,14 +18,7 @@ class ANIMNEXTSTATETREEUNCOOKEDONLY_API UAnimNextStateTreeTreeEditorData : publi
 public:
 	//~ Begin UStateTreeEditorData overrides
 	virtual void PostLoad() override;
-	
-	virtual const FInstancedPropertyBag& GetRootParametersPropertyBag() const override
-	{
-		const UStateTree* OuterStateTree = CastChecked<UStateTree>(GetOuter());
-		const UAnimNextStateTree* OuterAnimNextStateTree = CastChecked<UAnimNextStateTree>(OuterStateTree->GetOuter());
-		return OuterAnimNextStateTree->VariableDefaults;
-	}
-
+	virtual const FInstancedPropertyBag& GetRootParametersPropertyBag() const override;
 	virtual void CreateRootProperties(TArrayView<UE::PropertyBinding::FPropertyCreationDescriptor> InOutCreationDescs) override;
 	//~ End UStateTreeEditorData overrides
 
diff --git a/Engine/Plugins/Experimental/AnimNextWarping/Source/AnimNextWarping/Private/SteeringTrait.cpp b/Engine/Plugins/Experimental/AnimNextWarping/Source/AnimNextWarping/Private/SteeringTrait.cpp
--- a/Engine/Plugins/Experimental/AnimNextWarping/Source/AnimNextWarping/Private/SteeringTrait.cpp
+++ b/Engine/Plugins/Experimental/AnimNextWarping/Source/AnimNextWarping/Private/SteeringTrait.cpp
@@ -6,7 +6,6 @@
 #include "AnimNextWarpingLog.h"
 #include "Animation/AnimRootMotionProvider.h"
 #include "Animation/AnimSequenceBase.h"
-#include "DataInterface/DataInterfaceStructAdapter.h"
 #include "EvaluationVM/EvaluationVM.h"
 #include "EvaluationVM/KeyframeState.h"
 #include "TraitInterfaces/IAttributeProvider.h"
diff --git a/Engine/Plugins/Experimental/EvaluationNotifies/Source/Runtime/Private/AnimNext/AnimNextAlignment.cpp b/Engine/Plugins/Experimental/EvaluationNotifies/Source/Runtime/Private/AnimNext/AnimNextAlignment.cpp
--- a/Engine/Plugins/Experimental/EvaluationNotifies/Source/Runtime/Private/AnimNext/AnimNextAlignment.cpp
+++ b/Engine/Plugins/Experimental/EvaluationNotifies/Source/Runtime/Private/AnimNext/AnimNextAlignment.cpp
@@ -115,7 +115,7 @@ void FEvaluationNotify_AlignToGroundInstance::End(UE::AnimNext::FEvaluationNotif
 	FEvaluationNotify_AlignmentInstance::End(TraitInstanceData);
 	
 	const UNotifyState_AlignToGround* AlignToGroundNotify = CastChecked<UNotifyState_AlignToGround>(AnimNotify);
-	TraitInstanceData.DataInterface->SetVariable(AlignToGroundNotify->PlaybackRateOutputVariable, 1.0);
+	TraitInstanceData.Instance->SetVariable(AlignToGroundNotify->PlaybackRateOutputVariableReference, 1.0);
 }
 
 bool FEvaluationNotify_AlignToGroundInstance::GetTargetTransform(UE::AnimNext::FEvaluationNotifiesTrait::FInstanceData& TraitInstanceData, FTransform& Transform)
@@ -171,7 +171,7 @@ bool FEvaluationNotify_AlignToGroundInstance::GetTargetTransform(UE::AnimNext::F
 			static const FBox UnitBox(FVector(-10, -10, -10), FVector(10, 10, 10));
 			UE_VLOG_OBOX(TraitInstanceData.HostObject, "AlignToGround", Display, UnitBox, Transform.ToMatrixWithScale(), FColor::Red, TEXT(""));
 
-			if (!AlignToGroundNotify->PlaybackRateOutputVariable.IsNone())
+			if (!AlignToGroundNotify->PlaybackRateOutputVariableReference.IsNone())
 			{
 				float AnimatedFallDistance = StartTransform.GetTranslation().Z - EndTransform.GetTranslation().Z;
 				
@@ -188,7 +188,7 @@ bool FEvaluationNotify_AlignToGroundInstance::GetTargetTransform(UE::AnimNext::F
 				
 					float PlaybackRate = FMath::Clamp(PredictionDelta / ModifiedTime, AlignToGroundNotify->MinPlaybackRateModifier, AlignToGroundNotify->MaxPlaybackRateModifier);
 				
-					TraitInstanceData.DataInterface->SetVariable(AlignToGroundNotify->PlaybackRateOutputVariable, (double)PlaybackRate);
+					TraitInstanceData.Instance->SetVariable(AlignToGroundNotify->PlaybackRateOutputVariableReference, (double)PlaybackRate);
 				}
 			}
 			
@@ -219,8 +219,10 @@ FTransform GetModelSpaceTransform(const UE::AnimNext::FLODPoseStack& Pose, FBone
 bool FEvaluationNotify_AlignmentInstance::GetTargetTransform(UE::AnimNext::FEvaluationNotifiesTrait::FInstanceData& TraitInstanceData, FTransform& Transform)
 {
 	const UNotifyState_Alignment* AlignmentNotify = CastChecked<UNotifyState_Alignment>(AnimNotify);
-	
-	return TraitInstanceData.DataInterface->GetVariable(AlignmentNotify->TransformName, Transform) == EPropertyBagResult::Success;
+	// TODO: this should be a proper variable reference held on the notify, but currently this uses a 'named transform' system too. Needs resolving
+	PRAGMA_DISABLE_DEPRECATION_WARNINGS
+	return TraitInstanceData.Instance->GetVariable(FAnimNextVariableReference(AlignmentNotify->TransformName), Transform) == EPropertyBagResult::Success;
+	PRAGMA_ENABLE_DEPRECATION_WARNINGS
 }
 
 void FEvaluationNotify_AlignmentInstance::Update(UE::AnimNext::FEvaluationNotifiesTrait::FInstanceData& TraitInstanceData, UE::AnimNext::FEvaluationVM& VM)
@@ -240,7 +242,9 @@ void FEvaluationNotify_AlignmentInstance::Update(UE::AnimNext::FEvaluationNotifi
 			
 			if (AlignmentNotify->Disable != NAME_None)
 			{
-				TraitInstanceData.DataInterface->GetVariable(AlignmentNotify->Disable, bDisabled);
+				PRAGMA_DISABLE_DEPRECATION_WARNINGS
+				TraitInstanceData.Instance->GetVariable(FAnimNextVariableReference(AlignmentNotify->Disable), bDisabled);
+				PRAGMA_ENABLE_DEPRECATION_WARNINGS
 				
 				if (bDisabled)
 				{
diff --git a/Engine/Plugins/Experimental/EvaluationNotifies/Source/Runtime/Private/AnimNext/EvaluationNotifiesTrait.cpp b/Engine/Plugins/Experimental/EvaluationNotifies/Source/Runtime/Private/AnimNext/EvaluationNotifiesTrait.cpp
--- a/Engine/Plugins/Experimental/EvaluationNotifies/Source/Runtime/Private/AnimNext/EvaluationNotifiesTrait.cpp
+++ b/Engine/Plugins/Experimental/EvaluationNotifies/Source/Runtime/Private/AnimNext/EvaluationNotifiesTrait.cpp
@@ -119,8 +119,10 @@ namespace UE::AnimNext
 			InstanceData->OnExtractRootMotionAttribute = AttributeTrait.GetOnExtractRootMotionAttribute(Context);
 		}
 		
-		InstanceData->DataInterface = &Binding.GetTraitPtr().GetNodeInstance()->GetOwner();
-		InstanceData->DataInterface->GetVariable(RootBoneTransformName, InstanceData->RootBoneTransform);
+		InstanceData->Instance = &Binding.GetTraitPtr().GetNodeInstance()->GetOwner();
+		PRAGMA_DISABLE_DEPRECATION_WARNINGS
+		InstanceData->Instance->GetVariable(FAnimNextVariableReference(RootBoneTransformName), InstanceData->RootBoneTransform);
+		PRAGMA_ENABLE_DEPRECATION_WARNINGS
 
 		Context.AppendTask(FAnimNextEvaluationNotifiesTask::Make(InstanceData, SharedData));
 	}
diff --git a/Engine/Plugins/Experimental/EvaluationNotifies/Source/Runtime/Private/EvaluationNotifies/AnimNotifyState_Alignment.cpp b/Engine/Plugins/Experimental/EvaluationNotifies/Source/Runtime/Private/EvaluationNotifies/AnimNotifyState_Alignment.cpp
--- a/Engine/Plugins/Experimental/EvaluationNotifies/Source/Runtime/Private/EvaluationNotifies/AnimNotifyState_Alignment.cpp
+++ b/Engine/Plugins/Experimental/EvaluationNotifies/Source/Runtime/Private/EvaluationNotifies/AnimNotifyState_Alignment.cpp
@@ -5,6 +5,7 @@
 #include "Animation/AnimRootMotionProvider.h"
 #include "AHEasing/easing.h"
 #include "Animation/AnimSequenceHelpers.h"
+#include "UObject/FortniteMainBranchObjectVersion.h"
 #include "VisualLogger/VisualLogger.h"
 
 #include UE_INLINE_GENERATED_CPP_BY_NAME(AnimNotifyState_Alignment)
@@ -356,3 +357,18 @@ void FAlignmentNotifyInstance::Update(const UAnimSequenceBase* AnimationAsset, f
 	}
 }
 
+void UNotifyState_AlignToGround::Serialize(FArchive& Ar)
+{
+	Super::Serialize(Ar);
+
+	Ar.UsingCustomVersion(FFortniteMainBranchObjectVersion::GUID);
+
+#if WITH_EDITORONLY_DATA
+	if (Ar.IsLoading() && Ar.CustomVer(FFortniteMainBranchObjectVersion::GUID) < FFortniteMainBranchObjectVersion::AnimNextVariableReferences)
+	{
+		PRAGMA_DISABLE_DEPRECATION_WARNINGS
+		PlaybackRateOutputVariableReference = FAnimNextVariableReference(PlaybackRateOutputVariable_DEPRECATED);
+		PRAGMA_ENABLE_DEPRECATION_WARNINGS
+	}
+#endif
+}
diff --git a/Engine/Plugins/Experimental/EvaluationNotifies/Source/Runtime/Public/AnimNext/EvaluationNotifiesTrait.h b/Engine/Plugins/Experimental/EvaluationNotifies/Source/Runtime/Public/AnimNext/EvaluationNotifiesTrait.h
--- a/Engine/Plugins/Experimental/EvaluationNotifies/Source/Runtime/Public/AnimNext/EvaluationNotifiesTrait.h
+++ b/Engine/Plugins/Experimental/EvaluationNotifies/Source/Runtime/Public/AnimNext/EvaluationNotifiesTrait.h
@@ -3,7 +3,6 @@
 #pragma once
 
 #include "CoreMinimal.h"
-#include "DataInterface/AnimNextDataInterfaceInstance.h"
 #include "Kismet/KismetMathLibrary.h"
 #include "TraitCore/Trait.h"
 #include "TraitCore/TraitSharedData.h"
@@ -14,6 +13,7 @@
 #include "EvaluationNotifiesTrait.generated.h"
 
 struct FEvaluationNotify_BaseInstance;
+struct FUAFAssetInstance;
 
 /**
  * Data needed to execute steering
@@ -59,7 +59,7 @@ struct EVALUATIONNOTIFIESRUNTIME_API FEvaluationNotifiesTrait : FAdditiveTrait,
 		/** the evaluation notifies extracted from the current anim sequence */
 		TArray<FInstancedStruct> EvaluationNotifies;
 
-		FAnimNextDataInterfaceInstance* DataInterface = nullptr;
+		FUAFAssetInstance* Instance = nullptr;
 
 		/** Owner object the trait instance is associated to */
 		const UObject* HostObject;
diff --git a/Engine/Plugins/Experimental/EvaluationNotifies/Source/Runtime/Public/EvaluationNotifies/AnimNotifyState_Alignment.h b/Engine/Plugins/Experimental/EvaluationNotifies/Source/Runtime/Public/EvaluationNotifies/AnimNotifyState_Alignment.h
--- a/Engine/Plugins/Experimental/EvaluationNotifies/Source/Runtime/Public/EvaluationNotifies/AnimNotifyState_Alignment.h
+++ b/Engine/Plugins/Experimental/EvaluationNotifies/Source/Runtime/Public/EvaluationNotifies/AnimNotifyState_Alignment.h
@@ -3,6 +3,7 @@
 #pragma once
 
 #include "AnimNode_EvaluationNotifies.h"
+#include "Variables/AnimNextVariableReference.h"
 #include "AnimNotifyState_Alignment.generated.h"
 
 UENUM()
@@ -185,6 +186,9 @@ class UNotifyState_AlignToGround : public UNotifyState_AlignmentBase
 {
 	GENERATED_BODY()
 
+	// UObject interface
+	virtual void Serialize(FArchive& Ar) override;
+
 	public:
 	
 	UPROPERTY(EditAnywhere, Category = "Trace Settings")
@@ -202,11 +206,17 @@ class UNotifyState_AlignToGround : public UNotifyState_AlignmentBase
 	
 	UPROPERTY(EditAnywhere, Category = "Trace Settings")
 	TEnumAsByte<ETraceTypeQuery> TraceChannel = TraceTypeQuery1;
-	
+
+#if WITH_EDITORONLY_DATA
 	// Variable to output playback rate modifier to
-	UPROPERTY(EditAnywhere, Category = PlaybackRate)
-	FName PlaybackRateOutputVariable;
-	
+	UPROPERTY()
+	FName PlaybackRateOutputVariable_DEPRECATED;
+#endif
+
+	// Variable to output playback rate modifier to
+	UPROPERTY(EditAnywhere, DisplayName = "Playback Rate Output Variable", Category = PlaybackRate)
+	FAnimNextVariableReference PlaybackRateOutputVariableReference;
+
 	UPROPERTY(EditAnywhere, Category = PlaybackRate)
 	float MinPlaybackRateModifier = 0.5;
 	UPROPERTY(EditAnywhere, Category = PlaybackRate)
diff --git a/Engine/Plugins/Experimental/MovieSceneAnimMixer/Source/MovieSceneAnimMixer/Private/Systems/MovieSceneAnimNextTargetSystem.cpp b/Engine/Plugins/Experimental/MovieSceneAnimMixer/Source/MovieSceneAnimMixer/Private/Systems/MovieSceneAnimNextTargetSystem.cpp
--- a/Engine/Plugins/Experimental/MovieSceneAnimMixer/Source/MovieSceneAnimMixer/Private/Systems/MovieSceneAnimNextTargetSystem.cpp
+++ b/Engine/Plugins/Experimental/MovieSceneAnimMixer/Source/MovieSceneAnimMixer/Private/Systems/MovieSceneAnimNextTargetSystem.cpp
@@ -37,7 +37,7 @@
 #include "TraitInterfaces/IEvaluate.h"
 
 #include "PrimitiveSceneProxy.h"
-
+#include "UObject/FortniteMainBranchObjectVersion.h"
 
 #include UE_INLINE_GENERATED_CPP_BY_NAME(MovieSceneAnimNextTargetSystem)
 
@@ -66,7 +66,7 @@ namespace UE::MovieScene
 	/** Pre-animation traits for skeletal animations */
 	struct FPreAnimatedAnimNextMixedSkeletalAnimationTraits : FBoundObjectPreAnimatedStateTraits
 	{
-		using KeyType = TTuple<FObjectKey, FName>;
+		using KeyType = TTuple<FObjectKey, FAnimNextVariableReference>;
 		using StorageType = FPreAnimatedAnimNextState;
 
 		void RestorePreAnimatedValue(const KeyType& ObjectAndInjectionSite, StorageType& InOutCachedValue, const FRestoreStateParams& Params)
@@ -128,7 +128,7 @@ namespace UE::MovieScene
 	{
 		static TAutoRegisterPreAnimatedStorageID<FPreAnimatedAnimNextMixedSkeletalAnimationStorage> StorageID;
 
-		void BeginTrackingEntity(FMovieSceneEntityID EntityID, bool bWantsRestoreState, FRootInstanceHandle RootInstanceHandle, UAnimNextComponent* Component, FName InjectionSite)
+		void BeginTrackingEntity(FMovieSceneEntityID EntityID, bool bWantsRestoreState, FRootInstanceHandle RootInstanceHandle, UAnimNextComponent* Component, FAnimNextVariableReference InjectionSite)
 		{
 			if (!this->ParentExtension->IsCapturingGlobalState() && !bWantsRestoreState)
 			{
@@ -147,7 +147,7 @@ namespace UE::MovieScene
 		}
 
 		template<typename OnCacheValue /* StorageType(const KeyType&) */>
-		void CachePreAnimatedValue(UAnimNextComponent* Component, FName InjectionSite, OnCacheValue&& CacheCallback)
+		void CachePreAnimatedValue(UAnimNextComponent* Component, FAnimNextVariableReference InjectionSite, OnCacheValue&& CacheCallback)
 		{
 			EPreAnimatedCaptureSourceTracking TrackingMode = EPreAnimatedCaptureSourceTracking::CacheIfTracked;
 
@@ -244,15 +244,17 @@ namespace UE::MovieScene
 
 			if (AnimNextComponent)
 			{
-				UE::AnimNext::FInjectionSite InjectionSite(InjectionTarget->InjectionSiteName);
-				if (InjectionSite.DesiredSiteName == NAME_None)
+				UE::AnimNext::FInjectionSite InjectionSite(InjectionTarget->InjectionSite);
+				if (InjectionSite.DesiredSite.IsNone())
 				{
-					InjectionSite.DesiredSiteName = GSequencerDefaultAnimNextInjectionSite;
+					PRAGMA_DISABLE_DEPRECATION_WARNINGS
+					InjectionSite.DesiredSite = FAnimNextVariableReference(GSequencerDefaultAnimNextInjectionSite);
+					PRAGMA_ENABLE_DEPRECATION_WARNINGS
 					InjectionSite.bUseModuleFallback = true;
 				}
 
 				// If not currently injected, inject
-				FMovieSceneAnimNextTargetData* TargetData = System->CurrentTargets.FindByPredicate([AnimNextComponent, InjectionSite](const FMovieSceneAnimNextTargetData& Data) { return Data.AnimNextComponent == AnimNextComponent && Data.InjectionSite.DesiredSiteName == InjectionSite.DesiredSiteName;});
+				FMovieSceneAnimNextTargetData* TargetData = System->CurrentTargets.FindByPredicate([AnimNextComponent, InjectionSite](const FMovieSceneAnimNextTargetData& Data) { return Data.AnimNextComponent == AnimNextComponent && Data.InjectionSite.DesiredSite == InjectionSite.DesiredSite;});
 
 				if (!TargetData)
 				{
@@ -265,7 +267,7 @@ namespace UE::MovieScene
 					TargetData->InjectionRequestHandle = UE::AnimNext::FInjectionUtils::InjectEvaluationModifier(AnimNextComponent, TargetData->Modifier.ToSharedRef(), TargetData->InjectionSite);
 
 					// For now, we always restore state
-					PreAnimatedStorage->BeginTrackingEntity(EntityID, true, FRootInstanceHandle(), AnimNextComponent, InjectionSite.DesiredSiteName);
+					PreAnimatedStorage->BeginTrackingEntity(EntityID, true, FRootInstanceHandle(), AnimNextComponent, InjectionSite.DesiredSite);
 					PreAnimatedSkelMeshCompStorage->BeginTrackingEntity(EntityID, true, FRootInstanceHandle(), SkeletalMeshComponent);
 
 					bool bAnimNextModuleEnabled = AnimNextComponent->IsEnabled();
@@ -279,7 +281,7 @@ namespace UE::MovieScene
 						return State;
 					};
 
-					PreAnimatedStorage->CachePreAnimatedValue(AnimNextComponent, InjectionSite.DesiredSiteName, OnCacheAnimatedState);
+					PreAnimatedStorage->CachePreAnimatedValue(AnimNextComponent, InjectionSite.DesiredSite, OnCacheAnimatedState);
 					PreAnimatedSkelMeshCompStorage->CachePreAnimatedValue(FCachePreAnimatedValueParams(), SkeletalMeshComponent);
 
 					// Force enable the anim next module if not currently enabled
@@ -308,6 +310,24 @@ void FMovieSceneAnimMixerEvaluationModifier::PostEvaluate(UE::AnimNext::FEvaluat
 	}
 }
 
+#if WITH_EDITORONLY_DATA
+bool FMovieSceneAnimNextInjectionTarget::Serialize(FArchive& Ar)
+{
+	Ar.UsingCustomVersion(FFortniteMainBranchObjectVersion::GUID);
+	return false;
+}
+
+void FMovieSceneAnimNextInjectionTarget::PostSerialize(const FArchive& Ar)
+{
+	if (Ar.IsLoading() && Ar.CustomVer(FFortniteMainBranchObjectVersion::GUID) < FFortniteMainBranchObjectVersion::AnimNextVariableReferences)
+	{
+		PRAGMA_DISABLE_DEPRECATION_WARNINGS
+		InjectionSite = FAnimNextVariableReference(InjectionSiteName_DEPRECATED);
+		PRAGMA_ENABLE_DEPRECATION_WARNINGS
+	}
+}
+#endif
+
 UMovieSceneAnimNextTargetSystem::UMovieSceneAnimNextTargetSystem(const FObjectInitializer& ObjInit)
 	: Super(ObjInit)
 {
diff --git a/Engine/Plugins/Experimental/MovieSceneAnimMixer/Source/MovieSceneAnimMixer/Public/Systems/MovieSceneAnimNextTargetSystem.h b/Engine/Plugins/Experimental/MovieSceneAnimMixer/Source/MovieSceneAnimMixer/Public/Systems/MovieSceneAnimNextTargetSystem.h
--- a/Engine/Plugins/Experimental/MovieSceneAnimMixer/Source/MovieSceneAnimMixer/Public/Systems/MovieSceneAnimNextTargetSystem.h
+++ b/Engine/Plugins/Experimental/MovieSceneAnimMixer/Source/MovieSceneAnimMixer/Public/Systems/MovieSceneAnimNextTargetSystem.h
@@ -37,15 +37,22 @@ struct MOVIESCENEANIMMIXER_API FMovieSceneAnimNextInjectionTarget : public FMovi
 {
 	GENERATED_BODY() 
 
-	// Site name to use for injection. If empty will use the default one found for the actor.
-	// TODO: Custom UI that finds injection sites from modules.
-	UPROPERTY(EditAnywhere, Category="Animation")
-	FName InjectionSiteName;
+	UPROPERTY()
+	FName InjectionSiteName_DEPRECATED;
+
+	// Site to use for injection. If empty will use the default one found for the actor.
+	UPROPERTY(EditAnywhere, Category="Animation", meta = (AllowedType = "FAnimNextAnimGraph"))
+	FAnimNextVariableReference InjectionSite;
 
 	FORCEINLINE friend uint32 GetTypeHash(const FMovieSceneAnimNextInjectionTarget& Target)
 	{
-		return HashCombine(GetTypeHash(FMovieSceneAnimNextInjectionTarget::StaticStruct()), GetTypeHash(Target.InjectionSiteName));
+		return HashCombine(GetTypeHash(FMovieSceneAnimNextInjectionTarget::StaticStruct()), GetTypeHash(Target.InjectionSite));
 	}
+
+#if WITH_EDITORONLY_DATA
+	bool Serialize(FArchive& Ar);
+	void PostSerialize(const FArchive& Ar);
+#endif
 };
 
 struct FMovieSceneAnimNextTargetData
diff --git a/Engine/Source/Runtime/Core/Public/UObject/FortniteMainBranchObjectVersions.inl b/Engine/Source/Runtime/Core/Public/UObject/FortniteMainBranchObjectVersions.inl
--- a/Engine/Source/Runtime/Core/Public/UObject/FortniteMainBranchObjectVersions.inl
+++ b/Engine/Source/Runtime/Core/Public/UObject/FortniteMainBranchObjectVersions.inl
@@ -646,5 +646,8 @@ FORTNITE_MAIN_VERSION(MaterialFunctionBlendTopBottomInputEnum, 210)
 // Cooked CPU-side morph target points are now stored internally in the same compressed format as the GPU morph data. 
 FORTNITE_MAIN_VERSION(MorphTargetCookedCPUDataCompressed, 211)
 
+// AnimNext variables converted to references
+FORTNITE_MAIN_VERSION(AnimNextVariableReferences, 212)
+
 // -----<new versions can be added above this line>-------------------------------------------------
 #undef FORTNITE_MAIN_VERSION
\ No newline at end of file
