UAF: Fix crash creating a new AnimNextStateTree

Internal property bag is now held on the AnimNext asset EditorData, so needs to be in place when initially compiling the state tree

diff --git a/Engine/Plugins/Experimental/AnimNextStateTree/Source/AnimNextStateTreeEditor/Private/AnimNextStateTreeFactory.cpp b/Engine/Plugins/Experimental/AnimNextStateTree/Source/AnimNextStateTreeEditor/Private/AnimNextStateTreeFactory.cpp
--- a/Engine/Plugins/Experimental/AnimNextStateTree/Source/AnimNextStateTreeEditor/Private/AnimNextStateTreeFactory.cpp
+++ b/Engine/Plugins/Experimental/AnimNextStateTree/Source/AnimNextStateTreeEditor/Private/AnimNextStateTreeFactory.cpp
@@ -28,21 +28,20 @@ UObject* UAnimNextStateTreeFactory::FactoryCreateNew(UClass* Class, UObject* InP
 	}
 
 	UAnimNextStateTree* NewStateTree = NewObject<UAnimNextStateTree>(InParent, Class, Name, FlagsToUse);
+	// Create internal editor data
+	UAnimNextStateTree_EditorData* EditorData = NewObject<UAnimNextStateTree_EditorData>(NewStateTree, FName(Name.ToString() + "_StateTree_EditorData"), RF_Transactional);
+	NewStateTree->EditorData = EditorData;
 
-	// Create internal state tree	
+	// Create internal state tree
 	UStateTreeFactory* StateTreeFactory = NewObject<UStateTreeFactory>(GetTransientPackage());
 	StateTreeFactory->SetSchemaClass(UStateTreeAnimNextSchema::StaticClass());
 
 	UStateTree* StateTree = CastChecked<UStateTree>(StateTreeFactory->FactoryCreateNew(UStateTree::StaticClass(), NewStateTree, FName(Name.ToString() + "_StateTree"), RF_Transactional, nullptr, nullptr));
-
 	NewStateTree->StateTree = StateTree;
 
-	// Create internal editor data
-	UAnimNextStateTree_EditorData* EditorData = NewObject<UAnimNextStateTree_EditorData>(NewStateTree, FName(Name.ToString() + "_StateTree_EditorData"), RF_Transactional);
-	NewStateTree->EditorData = EditorData;
 	EditorData->bUsesExternalPackages = false;
 	EditorData->Initialize(/*bRecompileVM*/false);
-	
+
 	// Compile the initial skeleton
 	EditorData->RecompileVM();
 	
